Rename ClothLODDataBase to ClothLODDataCommon, and remove from the cloth runtime some related/minor dependencies to the nv module.
Also remove a bunch of unused get/set functions in the ClothMeshAdapter with unnecessary dependencies to the nv module.

diff --git a/Engine/Plugins/Experimental/ChaosCloth/Source/Chaos/Private/ChaosCloth/ChaosClothingSimulation.cpp b/Engine/Plugins/Experimental/ChaosCloth/Source/Chaos/Private/ChaosCloth/ChaosClothingSimulation.cpp
--- a/Engine/Plugins/Experimental/ChaosCloth/Source/Chaos/Private/ChaosCloth/ChaosClothingSimulation.cpp
+++ b/Engine/Plugins/Experimental/ChaosCloth/Source/Chaos/Private/ChaosCloth/ChaosClothingSimulation.cpp
@@ -169,7 +169,7 @@ void ClothingSimulation::CreateActor(USkeletalMeshComponent* InOwnerComponent, U
 	AnimDriveSpringStiffness[InSimDataIndex] = ChaosClothSimConfig->AnimDriveSpringStiffness;
 
 	check(Asset->GetNumLods() == 1);
-	UClothLODDataBase* AssetLodData = Asset->ClothLodData[0];
+	UClothLODDataCommon* AssetLodData = Asset->ClothLodData[0];
 	check(AssetLodData->PhysicalMeshData);
 	UClothPhysicalMeshDataBase* PhysMesh = AssetLodData->PhysicalMeshData;
 
@@ -739,7 +739,7 @@ void ClothingSimulation::ExtractLegacyAssetCollisions(UClothingAssetCommon* Asse
 	UE_CLOG(Asset->GetNumLods() != 1,
 		LogChaosCloth, Warning, TEXT("More than one LOD with the current cloth asset. Only LOD 0 is supported with the current system."));
 
-	if (const UClothLODDataBase* const AssetLodData = Asset->ClothLodData[0])
+	if (const UClothLODDataCommon* const AssetLodData = Asset->ClothLodData[0])
 	{
 		const FClothCollisionData& LodCollData = AssetLodData->CollisionData;
 		if (LodCollData.Spheres.Num() || LodCollData.SphereConnections.Num() || LodCollData.Convexes.Num())
@@ -1040,7 +1040,7 @@ void ClothingSimulation::Simulate(IClothingSimulationContext* InContext)
 		if (!Asset)
 			continue;
 
-		const UClothLODDataBase* AssetLodData = Asset->ClothLodData[0];
+		const UClothLODDataCommon* AssetLodData = Asset->ClothLodData[0];
 		check(AssetLodData->PhysicalMeshData);
 		const UClothPhysicalMeshDataBase* PhysMesh = AssetLodData->PhysicalMeshData;
 
@@ -1171,7 +1171,7 @@ void ClothingSimulation::GetCollisions(FClothCollisionData& OutCollisions, bool
 	// Add internal asset collisions
 	for (const UClothingAssetCommon* Asset : Assets)
 	{
-		if (const UClothLODDataBase* const ClothLodData = !Asset ? nullptr : Asset->ClothLodData[0])
+		if (const UClothLODDataCommon* const ClothLodData = !Asset ? nullptr : Asset->ClothLodData[0])
 		{
 			OutCollisions.Append(ClothLodData->CollisionData);
 		}
@@ -1553,7 +1553,7 @@ void ClothingSimulation::DebugDrawBackstops(USkeletalMeshComponent* OwnerCompone
 		}
 
 		// Get Backstop Distances
-		const UClothLODDataBase* const AssetLodData = Asset->ClothLodData[0];
+		const UClothLODDataCommon* const AssetLodData = Asset->ClothLodData[0];
 		check(AssetLodData);
 		check(AssetLodData->PhysicalMeshData);
 		const UClothPhysicalMeshDataBase* const PhysMesh = AssetLodData->PhysicalMeshData;
@@ -1610,7 +1610,7 @@ void ClothingSimulation::DebugDrawMaxDistances(USkeletalMeshComponent* OwnerComp
 		}
 
 		// Get Maximum Distances
-		const UClothLODDataBase* const AssetLodData = Asset->ClothLodData[0];
+		const UClothLODDataCommon* const AssetLodData = Asset->ClothLodData[0];
 		check(AssetLodData);
 		check(AssetLodData->PhysicalMeshData);
 		UClothPhysicalMeshDataBase* PhysMesh = AssetLodData->PhysicalMeshData;
@@ -1651,7 +1651,7 @@ void ClothingSimulation::DebugDrawAnimDrive(USkeletalMeshComponent* OwnerCompone
 		}
 
 		// Get Animdrive Multiplier
-		const UClothLODDataBase* const AssetLodData = Asset->ClothLodData[0];
+		const UClothLODDataCommon* const AssetLodData = Asset->ClothLodData[0];
 		check(AssetLodData);
 		check(AssetLodData->PhysicalMeshData);
 		const UClothPhysicalMeshDataBase* const PhysMesh = AssetLodData->PhysicalMeshData;
diff --git a/Engine/Source/Editor/ClothPainter/Private/ClothMeshAdapter.cpp b/Engine/Source/Editor/ClothPainter/Private/ClothMeshAdapter.cpp
--- a/Engine/Source/Editor/ClothPainter/Private/ClothMeshAdapter.cpp
+++ b/Engine/Source/Editor/ClothPainter/Private/ClothMeshAdapter.cpp
@@ -4,7 +4,6 @@
 
 #include "ClothingAssetBase.h"
 #include "ClothingAsset.h"
-#include "ClothPhysicalMeshDataNv.h" // for MaskTarget_PhysMesh
 
 #include "MeshPaintTypes.h"
 #include "Animation/DebugSkelMeshComponent.h"
@@ -143,105 +142,6 @@ FMatrix FClothMeshPaintAdapter::GetComponentToWorldMatrix() const
 	return SkeletalMeshComponent->GetComponentToWorld().ToMatrixWithScale();
 }
 
-float FClothMeshPaintAdapter::GetMaxDistanceValue(int32 VertexIndex) const
-{
-	for (const FClothAssetInfo& Info : AssetInfoMap)
-	{
-		if (VertexIndex >= Info.VertexStart && VertexIndex < Info.VertexEnd)
-		{
-			const TArray<float>* MaxDistances = 
-				Info.Asset->ClothLodData[PaintingClothLODIndex]->PhysicalMeshData->GetFloatArray(
-				(uint32)MaskTarget_PhysMesh::MaxDistance);
-			return MaxDistances ? (*MaxDistances)[VertexIndex - Info.VertexStart] : 0.0f;
-		}
-	}
-
-	return 0.0f;
-}
-
-void FClothMeshPaintAdapter::SetMaxDistanceValue(int32 VertexIndex, float Value)
-{
-	for (const FClothAssetInfo& Info : AssetInfoMap)
-	{
-		if (VertexIndex >= Info.VertexStart && VertexIndex < Info.VertexEnd)
-		{
-			if (TArray<float>* MaxDistances =
-				Info.Asset->ClothLodData[PaintingClothLODIndex]->PhysicalMeshData->GetFloatArray(
-				(uint32)MaskTarget_PhysMesh::MaxDistance))
-			{
-				(*MaxDistances)[VertexIndex - Info.VertexStart] = Value;
-			}
-			break;
-		}
-	}
-}
-
-float FClothMeshPaintAdapter::GetBackstopDistanceValue(int32 VertexIndex) const
-{
-	for (const FClothAssetInfo& Info : AssetInfoMap)
-	{
-		if (VertexIndex >= Info.VertexStart && VertexIndex < Info.VertexEnd)
-		{
-			const TArray<float>* BackstopDistances = 
-				Info.Asset->ClothLodData[PaintingClothLODIndex]->PhysicalMeshData->GetFloatArray(
-				(uint32)MaskTarget_PhysMesh::BackstopDistance);
-			return BackstopDistances ? (*BackstopDistances)[VertexIndex - Info.VertexStart] : 0.0f;
-		}
-	}
-
-	return 0.0f;
-}
-
-void FClothMeshPaintAdapter::SetBackstopDistanceValue(int32 VertexIndex, float Value)
-{
-	for (const FClothAssetInfo& Info : AssetInfoMap)
-	{
-		if (VertexIndex >= Info.VertexStart && VertexIndex < Info.VertexEnd)
-		{
-			if (TArray<float>* BackstopDistances =
-				Info.Asset->ClothLodData[PaintingClothLODIndex]->PhysicalMeshData->GetFloatArray(
-				(uint32)MaskTarget_PhysMesh::BackstopDistance))
-			{
-				(*BackstopDistances)[VertexIndex - Info.VertexStart] = Value;
-			}
-			break;
-		}
-	}
-}
-
-float FClothMeshPaintAdapter::GetBackstopRadiusValue(int32 VertexIndex) const
-{
-	for (const FClothAssetInfo& Info : AssetInfoMap)
-	{
-		if (VertexIndex >= Info.VertexStart && VertexIndex < Info.VertexEnd)
-		{
-			const TArray<float>* BackstopRadii = 
-				Info.Asset->ClothLodData[PaintingClothLODIndex]->PhysicalMeshData->GetFloatArray(
-				(uint32)MaskTarget_PhysMesh::BackstopRadius);
-			return BackstopRadii ? (*BackstopRadii)[VertexIndex - Info.VertexStart] : 0.0f;
-		}
-	}
-
-	return 0.0f;
-}
-
-void FClothMeshPaintAdapter::SetBackstopRadiusValue(int32 VertexIndex, float Value)
-{
-	for (const FClothAssetInfo& Info : AssetInfoMap)
-	{
-		if (VertexIndex >= Info.VertexStart && VertexIndex < Info.VertexEnd)
-		{
-			if (TArray<float>* BackstopRadii =
-				Info.Asset->ClothLodData[PaintingClothLODIndex]->PhysicalMeshData->GetFloatArray(
-				(uint32)MaskTarget_PhysMesh::BackstopRadius))
-			{
-				(*BackstopRadii)[VertexIndex - Info.VertexStart] = Value;
-			}
-			break;
-		}
-	}
-}
-
 TArray<FVector> FClothMeshPaintAdapter::SphereIntersectVertices(const float ComponentSpaceSquaredBrushRadius, const FVector& ComponentSpaceBrushPosition, const FVector& ComponentSpaceCameraPosition, const bool bOnlyFrontFacing) const
 {
 	// Get list of intersecting triangles with given sphere data
@@ -282,7 +182,7 @@ void FClothMeshPaintAdapter::SetSelectedClothingAsset(const FGuid& InAssetGuid,
 			{
 				if(ConcreteAsset->IsValidLod(InAssetLod))
 				{
-					UClothLODDataBase* LodData = ConcreteAsset->ClothLodData[InAssetLod];
+					UClothLODDataCommon* LodData = ConcreteAsset->ClothLodData[InAssetLod];
 
 					if(LodData->ParameterMasks.IsValidIndex(InMaskIndex))
 					{
@@ -333,7 +233,7 @@ bool FClothMeshPaintAdapter::InitializeVertexData()
 			if(DebugComponent->SkinnedSelectedClothingPositions.Num() > 0)
 			{
 				UClothingAssetCommon* ConcreteAsset = CastChecked<UClothingAssetCommon>(SelectedAsset);
-				const UClothLODDataBase* LODData = ConcreteAsset->ClothLodData[PaintingClothLODIndex];
+				const UClothLODDataCommon* LODData = ConcreteAsset->ClothLodData[PaintingClothLODIndex];
 				check(LODData->PhysicalMeshData);
 				const UClothPhysicalMeshDataBase* MeshData = LODData->PhysicalMeshData;
 
diff --git a/Engine/Source/Editor/ClothPainter/Private/ClothMeshAdapter.h b/Engine/Source/Editor/ClothPainter/Private/ClothMeshAdapter.h
--- a/Engine/Source/Editor/ClothPainter/Private/ClothMeshAdapter.h
+++ b/Engine/Source/Editor/ClothPainter/Private/ClothMeshAdapter.h
@@ -64,21 +64,6 @@ public:
 
 	virtual TArray<FVector> SphereIntersectVertices(const float ComponentSpaceSquaredBrushRadius, const FVector& ComponentSpaceBrushPosition, const FVector& ComponentSpaceCameraPosition, const bool bOnlyFrontFacing) const;
 
-	/** Retrieves the backstop distance value for the given vertex index from the simulation data */
-	virtual float GetBackstopDistanceValue(int32 VertexIndex) const;
-	/** Sets the backstop distance value for the given vertex index to Value */
-	virtual void SetBackstopDistanceValue(int32 VertexIndex, float Value);
-
-	/** Retrieves the backstop radius value for the given vertex index from the simulation data */
-	virtual float GetBackstopRadiusValue(int32 VertexIndex) const;
-	/** Sets the backstop radius value for the given vertex index to Value */
-	virtual void SetBackstopRadiusValue(int32 VertexIndex, float Value);
-
-	/** Retrieves the max distance value for the given vertex index from the simulation data */
-	virtual float GetMaxDistanceValue(int32 VertexIndex) const;
-	/** Sets the max distance value for the given vertex index to Value */
-	virtual void SetMaxDistanceValue(int32 VertexIndex, float Value);
-
 	/** Sets the represented clothing asset to the UClothingAssetBase retrieved from the AssetGUID */
 	virtual void SetSelectedClothingAsset(const FGuid& InAssetGuid, int32 InAssetLod, int32 InMaskIndex);
 
diff --git a/Engine/Source/Editor/ClothPainter/Private/SClothAssetSelector.cpp b/Engine/Source/Editor/ClothPainter/Private/SClothAssetSelector.cpp
--- a/Engine/Source/Editor/ClothPainter/Private/SClothAssetSelector.cpp
+++ b/Engine/Source/Editor/ClothPainter/Private/SClothAssetSelector.cpp
@@ -39,7 +39,7 @@ FPointWeightMap* FClothingMaskListItem::GetMask()
 	{
 		if(Asset->IsValidLod(LodIndex))
 		{
-			UClothLODDataBase* LodData = Asset->ClothLodData[LodIndex];
+			UClothLODDataCommon* LodData = Asset->ClothLodData[LodIndex];
 			if(LodData->ParameterMasks.IsValidIndex(MaskIndex))
 			{
 				return &LodData->ParameterMasks[MaskIndex];
@@ -304,8 +304,8 @@ private:
 
 			for(int32 CurrIndex = 0; CurrIndex < NumLods - 1; ++CurrIndex)
 			{
-				UClothLODDataBase* SourceLod = Asset->ClothLodData[CurrIndex];
-				UClothLODDataBase* DestLod = Asset->ClothLodData[CurrIndex + 1];
+				UClothLODDataCommon* SourceLod = Asset->ClothLodData[CurrIndex];
+				UClothLODDataCommon* DestLod = Asset->ClothLodData[CurrIndex + 1];
 
 				DestLod->ParameterMasks.Reset();
 
@@ -492,7 +492,7 @@ private:
 		);
 	}
 
-	UClothLODDataBase* GetCurrentLod() const
+	UClothLODDataCommon* GetCurrentLod() const
 	{
 		if(Item.IsValid())
 		{
@@ -500,7 +500,7 @@ private:
 			{
 				if(Asset->ClothLodData.IsValidIndex(Item->LodIndex))
 				{
-					UClothLODDataBase* LodData = Asset->ClothLodData[Item->LodIndex];
+					UClothLODDataCommon* LodData = Asset->ClothLodData[Item->LodIndex];
 					return LodData;
 				}
 			}
@@ -518,7 +518,7 @@ private:
 			FScopedTransaction CurrTransaction(LOCTEXT("DeleteMask_Transaction", "Delete clothing parameter mask."));
 			Item->ClothingAsset->Modify();
 
-		if(UClothLODDataBase* LodData = GetCurrentLod())
+		if(UClothLODDataCommon* LodData = GetCurrentLod())
 		{
 			if(LodData->ParameterMasks.IsValidIndex(Item->MaskIndex))
 			{
@@ -645,7 +645,7 @@ private:
 						{
 							if(Asset->ClothLodData.IsValidIndex(InItem->LodIndex))
 							{
-								UClothLODDataBase* LodData = Asset->ClothLodData[InItem->LodIndex];
+								UClothLODDataCommon* LodData = Asset->ClothLodData[InItem->LodIndex];
 
 								TArray<FPointWeightMap*> AllTargetMasks;
 								LodData->GetParameterMasksForTarget(Mask->CurrentTarget, AllTargetMasks);
@@ -1129,7 +1129,7 @@ FReply SClothAssetSelector::AddNewMask()
 	{
 		if(Asset->ClothLodData.IsValidIndex(SelectedLod))
 		{
-			UClothLODDataBase* LodData = Asset->ClothLodData[SelectedLod];
+			UClothLODDataCommon* LodData = Asset->ClothLodData[SelectedLod];
 			const int32 NumRequiredValues = LodData->PhysicalMeshData->Vertices.Num();
 
 			LodData->ParameterMasks.AddDefaulted();
@@ -1233,7 +1233,7 @@ void SClothAssetSelector::RefreshMaskList()
 	UClothingAssetCommon* Asset = SelectedAsset.Get();
 	if(Asset && Asset->IsValidLod(SelectedLod))
 	{
-		UClothLODDataBase* LodData = Asset->ClothLodData[SelectedLod];
+		UClothLODDataCommon* LodData = Asset->ClothLodData[SelectedLod];
 		const int32 NumMasks = LodData->ParameterMasks.Num();
 
 		for(int32 Index = 0; Index < NumMasks; ++Index)
@@ -1285,7 +1285,7 @@ void SClothAssetSelector::OnClothingLodSelected(int32 InNewLod)
 		int32 NewMaskSelection = INDEX_NONE;
 		if(SelectedAsset->ClothLodData.IsValidIndex(SelectedLod))
 		{
-			UClothLODDataBase* LodData = SelectedAsset->ClothLodData[SelectedLod];
+			UClothLODDataCommon* LodData = SelectedAsset->ClothLodData[SelectedLod];
 
 			if(LodData->ParameterMasks.Num() > 0)
 			{
@@ -1309,7 +1309,7 @@ void SClothAssetSelector::SetSelectedAsset(TWeakObjectPtr<UClothingAssetCommon>
 		{
 			SetSelectedLod(0);
 
-			UClothLODDataBase* LodData = NewAsset->ClothLodData[SelectedLod];
+			UClothLODDataCommon* LodData = NewAsset->ClothLodData[SelectedLod];
 			if(LodData->ParameterMasks.Num() > 0)
 			{
 				SetSelectedMask(0);
diff --git a/Engine/Source/Editor/ClothPainter/Private/SCopyVertexColorSettingsPanel.cpp b/Engine/Source/Editor/ClothPainter/Private/SCopyVertexColorSettingsPanel.cpp
--- a/Engine/Source/Editor/ClothPainter/Private/SCopyVertexColorSettingsPanel.cpp
+++ b/Engine/Source/Editor/ClothPainter/Private/SCopyVertexColorSettingsPanel.cpp
@@ -112,7 +112,7 @@ FReply SCopyVertexColorSettingsPanel::OnCopyClicked()
 	UClothingAssetCommon* Asset = SelectedAssetPtr.Get();
 	if (Asset && SelectedMask)
 	{
-		UClothLODDataBase* ClothLODData = Asset->ClothLodData[SelectedLOD];
+		UClothLODDataCommon* ClothLODData = Asset->ClothLodData[SelectedLOD];
 
 		check(ClothLODData->PhysicalMeshData->Vertices.Num() == ClothLODData->PhysicalMeshData->VertexColors.Num());
 		int32 NumVerts = ClothLODData->PhysicalMeshData->Vertices.Num();
diff --git a/Engine/Source/Editor/ClothingSystemEditor/Private/ClothingAssetFactory.cpp b/Engine/Source/Editor/ClothingSystemEditor/Private/ClothingAssetFactory.cpp
--- a/Engine/Source/Editor/ClothingSystemEditor/Private/ClothingAssetFactory.cpp
+++ b/Engine/Source/Editor/ClothingSystemEditor/Private/ClothingAssetFactory.cpp
@@ -248,7 +248,7 @@ UClothingAssetBase* UClothingAssetFactory::CreateFromSkeletalMesh(USkeletalMesh*
 
 	// Adding a new LOD from this skeletal mesh
 	NewAsset->AddNewLod();
-	UClothLODDataBase* LodData = NewAsset->ClothLodData.Last();
+	UClothLODDataCommon* LodData = NewAsset->ClothLodData.Last();
 
 	if(ImportToLodInternal(TargetMesh, Params.LodIndex, Params.SourceSection, NewAsset, LodData))
 	{
@@ -327,7 +327,7 @@ UClothingAssetBase* UClothingAssetFactory::ImportLodToClothing(USkeletalMesh* Ta
 			// Everything looks good, continue to actual import
 			UClothingAssetCommon* ConcreteTarget = CastChecked<UClothingAssetCommon>(TargetClothing);
 
-			UClothLODDataBase* RemapSource = nullptr;
+			UClothLODDataCommon* RemapSource = nullptr;
 
 			if(Params.bRemapParameters)
 			{
@@ -354,7 +354,7 @@ UClothingAssetBase* UClothingAssetFactory::ImportLodToClothing(USkeletalMesh* Ta
 				return nullptr;
 			}
 
-			UClothLODDataBase* NewLod = ConcreteTarget->ClothLodData[Params.TargetLod];
+			UClothLODDataCommon* NewLod = ConcreteTarget->ClothLodData[Params.TargetLod];
 
 			if(Params.TargetLod > 0 && Params.bRemapParameters)
 			{
@@ -403,12 +403,12 @@ UClothingAssetBase* UClothingAssetFactory::CreateFromApexAsset(nvidia::apex::Clo
 	for(int32 CurrLodIdx = 0; CurrLodIdx < NumLodsToBuild; ++CurrLodIdx)
 	{
 		NewClothingAsset->AddNewLod();
-		UClothLODDataNv* CurrentLodDataPtr = Cast<UClothLODDataNv>(NewClothingAsset->ClothLodData[CurrLodIdx]);
+		UClothLODDataCommon* CurrentLodDataPtr = NewClothingAsset->ClothLodData[CurrLodIdx];
 		if (!ensure(CurrentLodDataPtr))
 		{
 			continue;
 		}
-		UClothLODDataNv& CurrentLodData = *CurrentLodDataPtr;
+		UClothLODDataCommon& CurrentLodData = *CurrentLodDataPtr;
 
 		TArray<FApexVertData> ApexVertData;
 
@@ -476,7 +476,7 @@ UClothingAssetBase* UClothingAssetFactory::CreateFromApexAsset(nvidia::apex::Clo
 	NewClothingAsset->CalculateReferenceBoneIndex();
 
 	// Add masks for parameters
-	for (UClothLODDataBase* LodPtr : NewClothingAsset->ClothLodData)
+	for (UClothLODDataCommon* LodPtr : NewClothingAsset->ClothLodData)
 	{
 		check(LodPtr);
 		UClothPhysicalMeshDataNv* PhysMeshPtr = Cast<UClothPhysicalMeshDataNv>(LodPtr->PhysicalMeshData);
@@ -856,7 +856,7 @@ void UClothingAssetFactory::ExtractBoneData(UClothingAssetCommon* NewAsset, Clot
 	}
 }
 
-void UClothingAssetFactory::ExtractSphereCollisions(UClothingAssetCommon* NewAsset, nvidia::apex::ClothingAsset &InApexAsset, int32 InLodIdx, UClothLODDataNv &InLodData)
+void UClothingAssetFactory::ExtractSphereCollisions(UClothingAssetCommon* NewAsset, nvidia::apex::ClothingAsset &InApexAsset, int32 InLodIdx, UClothLODDataCommon &InLodData)
 {
 	const NvParameterized::Interface* AssetParams = InApexAsset.getAssetNvParameterized();
 
@@ -1139,8 +1139,8 @@ bool UClothingAssetFactory::ImportToLodInternal(
 	int32 SourceLodIndex, 
 	int32 SourceSectionIndex, 
 	UClothingAssetCommon* DestAsset, 
-	UClothLODDataBase* DestLod, 
-	UClothLODDataBase* InParameterRemapSource)
+	UClothLODDataCommon* DestLod, 
+	UClothLODDataCommon* InParameterRemapSource)
 {
 	if(!SourceMesh || !SourceMesh->GetImportedModel())
 	{
@@ -1339,7 +1339,7 @@ bool UClothingAssetFactory::ImportToLodInternal(
 
 #if WITH_APEX_CLOTHING
 
-void UClothingAssetFactory::ExtractLodPhysicalData(UClothingAssetCommon* NewAsset, ClothingAsset &InApexAsset, int32 InLodIdx, UClothLODDataNv &InLodData, TArray<FApexVertData>& OutApexVertData)
+void UClothingAssetFactory::ExtractLodPhysicalData(UClothingAssetCommon* NewAsset, ClothingAsset &InApexAsset, int32 InLodIdx, UClothLODDataCommon &InLodData, TArray<FApexVertData>& OutApexVertData)
 {
 	const NvParameterized::Interface* AssetParams = InApexAsset.getAssetNvParameterized();
 
diff --git a/Engine/Source/Editor/ClothingSystemEditor/Public/ClothingAssetFactory.h b/Engine/Source/Editor/ClothingSystemEditor/Public/ClothingAssetFactory.h
--- a/Engine/Source/Editor/ClothingSystemEditor/Public/ClothingAssetFactory.h
+++ b/Engine/Source/Editor/ClothingSystemEditor/Public/ClothingAssetFactory.h
@@ -65,16 +65,16 @@ private:
 
 	// Extraction methods for pulling the required data from an APEX asset and
 	// pushing it to a UClothingAssetCommon
-	void ExtractLodPhysicalData(UClothingAssetCommon* NewAsset, nvidia::apex::ClothingAsset &InApexAsset, int32 InLodIdx, UClothLODDataNv &InLodData, TArray<FApexVertData>& OutApexVertData);
+	void ExtractLodPhysicalData(UClothingAssetCommon* NewAsset, nvidia::apex::ClothingAsset &InApexAsset, int32 InLodIdx, UClothLODDataCommon &InLodData, TArray<FApexVertData>& OutApexVertData);
 	void ExtractBoneData(UClothingAssetCommon* NewAsset, nvidia::apex::ClothingAsset &InApexAsset);
-	void ExtractSphereCollisions(UClothingAssetCommon* NewAsset, nvidia::apex::ClothingAsset &InApexAsset, int32 InLodIdx, UClothLODDataNv &InLodData);
+	void ExtractSphereCollisions(UClothingAssetCommon* NewAsset, nvidia::apex::ClothingAsset &InApexAsset, int32 InLodIdx, UClothLODDataCommon &InLodData);
 	void ExtractMaterialParameters(UClothingAssetCommon* NewAsset, nvidia::apex::ClothingAsset &InApexAsset);
 #endif
 
 	// Utility methods for skeletal mesh extraction //////////////////////////
 
 	/** Handles internal import of LODs */
-	bool ImportToLodInternal(USkeletalMesh* SourceMesh, int32 SourceLodIndex, int32 SourceSectionIndex, UClothingAssetCommon* DestAsset, UClothLODDataBase* DestLod, UClothLODDataBase* InParameterRemapSource = nullptr);
+	bool ImportToLodInternal(USkeletalMesh* SourceMesh, int32 SourceLodIndex, int32 SourceSectionIndex, UClothingAssetCommon* DestAsset, UClothLODDataCommon* DestLod, UClothLODDataCommon* InParameterRemapSource = nullptr);
 
 	//////////////////////////////////////////////////////////////////////////
 
diff --git a/Engine/Source/Editor/Persona/Private/PersonaMeshDetails.cpp b/Engine/Source/Editor/Persona/Private/PersonaMeshDetails.cpp
--- a/Engine/Source/Editor/Persona/Private/PersonaMeshDetails.cpp
+++ b/Engine/Source/Editor/Persona/Private/PersonaMeshDetails.cpp
@@ -5156,7 +5156,7 @@ TSharedRef<SUniformGridPanel> FPersonaMeshDetails::MakeClothingDetailsWidget(int
 
 		if (UClothingAssetCommon* Asset = Cast<UClothingAssetCommon>(ClothingAsset))
 		{
-			UClothLODDataBase* LodData = Asset->ClothLodData[LODIndex];
+			UClothLODDataCommon* LodData = Asset->ClothLodData[LODIndex];
 			check(LodData->PhysicalMeshData);
 			UClothPhysicalMeshDataBase& PhysMeshData = *LodData->PhysicalMeshData;
 			FClothCollisionData& CollisionData = LodData->CollisionData;
diff --git a/Engine/Source/Editor/UnrealEd/Private/Animation/DebugSkelMeshComponent.cpp b/Engine/Source/Editor/UnrealEd/Private/Animation/DebugSkelMeshComponent.cpp
--- a/Engine/Source/Editor/UnrealEd/Private/Animation/DebugSkelMeshComponent.cpp
+++ b/Engine/Source/Editor/UnrealEd/Private/Animation/DebugSkelMeshComponent.cpp
@@ -741,7 +741,7 @@ void UDebugSkelMeshComponent::RebuildClothingSectionsFixedVerts()
 				if(BaseAsset)
 				{
 					UClothingAssetCommon* ConcreteAsset = Cast<UClothingAssetCommon>(BaseAsset);
-					UClothLODDataBase* LodData = ConcreteAsset->ClothLodData[Section.ClothingData.AssetLodIndex];
+					UClothLODDataCommon* LodData = ConcreteAsset->ClothLodData[Section.ClothingData.AssetLodIndex];
 
 					for(FMeshToMeshVertData& VertData : Section.ClothMappingData)
 					{
@@ -828,7 +828,7 @@ void UDebugSkelMeshComponent::RefreshSelectedClothingSkinnedPositions()
 				// Pass LOD0 to collect all bones
 				GetCurrentRefToLocalMatrices(RefToLocals, 0);
 
-				UClothLODDataBase* LodData = ConcreteAsset->ClothLodData[SelectedClothingLodForPainting];
+				UClothLODDataCommon* LodData = ConcreteAsset->ClothLodData[SelectedClothingLodForPainting];
 
 				ClothingMeshUtils::SkinPhysicsMesh(ConcreteAsset->UsedBoneIndices, *LodData->PhysicalMeshData, FTransform::Identity, RefToLocals.GetData(), RefToLocals.Num(), SkinnedSelectedClothingPositions, SkinnedSelectedClothingNormals);
 				RebuildCachedClothBounds();
@@ -1035,7 +1035,7 @@ FDebugSkelMeshDynamicData::FDebugSkelMeshDynamicData(UDebugSkelMeshComponent* In
 					{
 						if(ConcreteAsset->ClothLodData.IsValidIndex(InComponent->SelectedClothingLodForPainting))
 						{
-							UClothLODDataBase* LodData = ConcreteAsset->ClothLodData[InComponent->SelectedClothingLodForPainting];
+							UClothLODDataCommon* LodData = ConcreteAsset->ClothLodData[InComponent->SelectedClothingLodForPainting];
 
 							ClothingSimIndices = LodData->PhysicalMeshData->Indices;
 
diff --git a/Engine/Source/Runtime/ClothingSystemRuntimeCommon/Private/ClothLODData.cpp b/Engine/Source/Runtime/ClothingSystemRuntimeCommon/Private/ClothLODData.cpp
--- a/Engine/Source/Runtime/ClothingSystemRuntimeCommon/Private/ClothLODData.cpp
+++ b/Engine/Source/Runtime/ClothingSystemRuntimeCommon/Private/ClothLODData.cpp
@@ -3,15 +3,15 @@
 #include "ClothLODData.h"
 
 
-UClothLODDataBase::UClothLODDataBase(const FObjectInitializer& Init)
+UClothLODDataCommon::UClothLODDataCommon(const FObjectInitializer& Init)
 	: PhysicalMeshData(nullptr)
 {}
 
-UClothLODDataBase::~UClothLODDataBase()
+UClothLODDataCommon::~UClothLODDataCommon()
 {}
 
 #if WITH_EDITORONLY_DATA
-void UClothLODDataBase::GetParameterMasksForTarget(
+void UClothLODDataCommon::GetParameterMasksForTarget(
 	const uint8 InTarget, 
 	TArray<FPointWeightMap*>& OutMasks)
 {
@@ -25,7 +25,7 @@ void UClothLODDataBase::GetParameterMasksForTarget(
 }
 #endif // WITH_EDITORONLY_DATA
 
-void UClothLODDataBase::Serialize(FArchive& Ar)
+void UClothLODDataCommon::Serialize(FArchive& Ar)
 {
 	Super::Serialize(Ar);
 
@@ -46,7 +46,7 @@ void UClothLODDataBase::Serialize(FArchive& Ar)
 }
 
 #if WITH_EDITOR
-void UClothLODDataBase::PushWeightsToMesh()
+void UClothLODDataCommon::PushWeightsToMesh()
 {
 	if (PhysicalMeshData)
 	{
diff --git a/Engine/Source/Runtime/ClothingSystemRuntimeCommon/Private/ClothingAsset.cpp b/Engine/Source/Runtime/ClothingSystemRuntimeCommon/Private/ClothingAsset.cpp
--- a/Engine/Source/Runtime/ClothingSystemRuntimeCommon/Private/ClothingAsset.cpp
+++ b/Engine/Source/Runtime/ClothingSystemRuntimeCommon/Private/ClothingAsset.cpp
@@ -265,7 +265,7 @@ bool UClothingAssetCommon::BindToSkeletalMesh(
 	CalculateReferenceBoneIndex();
 
 	// Grab the clothing and skel lod data
-	UClothLODDataBase* LodData = ClothLodData[InAssetLodIndex];
+	UClothLODDataCommon* LodData = ClothLodData[InAssetLodIndex];
 	FSkeletalMeshLODModel& SkelLod = InSkelMesh->GetImportedModel()->LODModels[InMeshLodIndex];
 
 	FSkelMeshSection& OriginalSection = SkelLod.Sections[InSectionIndex];
@@ -505,7 +505,7 @@ void UClothingAssetCommon::ForEachInteractorUsingClothing(TFunction<void(UClothi
 
 void UClothingAssetCommon::ApplyParameterMasks()
 {
-	for(UClothLODDataBase* Lod : ClothLodData)
+	for(UClothLODDataCommon* Lod : ClothLodData)
 	{
 		Lod->PushWeightsToMesh();
 	}
@@ -520,12 +520,12 @@ void UClothingAssetCommon::BuildLodTransitionData()
 		const bool bHasPrevLod = LodIndex > 0;
 		const bool bHasNextLod = LodIndex < NumLods - 1;
 
-		UClothLODDataBase* CurrentLod = ClothLodData[LodIndex];
+		UClothLODDataCommon* CurrentLod = ClothLodData[LodIndex];
 		check(CurrentLod->PhysicalMeshData);
 		UClothPhysicalMeshDataBase& CurrentPhysMesh = *CurrentLod->PhysicalMeshData;
 
-		UClothLODDataBase* PrevLod = bHasPrevLod ? ClothLodData[LodIndex - 1] : nullptr;
-		UClothLODDataBase* NextLod = bHasNextLod ? ClothLodData[LodIndex + 1] : nullptr;
+		UClothLODDataCommon* PrevLod = bHasPrevLod ? ClothLodData[LodIndex - 1] : nullptr;
+		UClothLODDataCommon* NextLod = bHasNextLod ? ClothLodData[LodIndex + 1] : nullptr;
 
 		const int32 CurrentLodNumVerts = CurrentPhysMesh.Vertices.Num();
 
@@ -597,7 +597,7 @@ void UClothingAssetCommon::CalculateReferenceBoneIndex()
 		// List of actually weighted (not just used) bones
 		TArray<int32> WeightedBones;
 
-		for(UClothLODDataBase* CurLod : ClothLodData)
+		for(UClothLODDataCommon* CurLod : ClothLodData)
 		{
 			check(CurLod && CurLod->PhysicalMeshData);
 			UClothPhysicalMeshDataBase* MeshData = CurLod->PhysicalMeshData;
@@ -700,7 +700,7 @@ void UClothingAssetCommon::BuildSelfCollisionData()
 		// No self collision, can't generate data
 		return;
 	}
-	for(UClothLODDataBase* Lod : ClothLodData)
+	for(UClothLODDataCommon* Lod : ClothLodData)
 	{
 		check(Lod && Lod->PhysicalMeshData);
 		Lod->PhysicalMeshData->BuildSelfCollisionData(ClothSimConfig);
diff --git a/Engine/Source/Runtime/ClothingSystemRuntimeCommon/Public/ClothLODData.h b/Engine/Source/Runtime/ClothingSystemRuntimeCommon/Public/ClothLODData.h
--- a/Engine/Source/Runtime/ClothingSystemRuntimeCommon/Public/ClothLODData.h
+++ b/Engine/Source/Runtime/ClothingSystemRuntimeCommon/Public/ClothLODData.h
@@ -12,12 +12,12 @@
 
 /** */
 UCLASS()
-class CLOTHINGSYSTEMRUNTIMECOMMON_API UClothLODDataBase : public UObject
+class CLOTHINGSYSTEMRUNTIMECOMMON_API UClothLODDataCommon : public UObject
 {
 	GENERATED_BODY()
 public:
-	UClothLODDataBase(const FObjectInitializer& Init);
-	virtual ~UClothLODDataBase();
+	UClothLODDataCommon(const FObjectInitializer& Init);
+	virtual ~UClothLODDataCommon();
 
 	// Raw phys mesh data
 	UPROPERTY(EditAnywhere, Category = SimMesh)
diff --git a/Engine/Source/Runtime/ClothingSystemRuntimeCommon/Public/ClothingAsset.h b/Engine/Source/Runtime/ClothingSystemRuntimeCommon/Public/ClothingAsset.h
--- a/Engine/Source/Runtime/ClothingSystemRuntimeCommon/Public/ClothingAsset.h
+++ b/Engine/Source/Runtime/ClothingSystemRuntimeCommon/Public/ClothingAsset.h
@@ -178,7 +178,7 @@ public:
 
 	// The actual asset data, listed by LOD.
 	UPROPERTY()
-	TArray<UClothLODDataBase*> ClothLodData;
+	TArray<UClothLODDataCommon*> ClothLodData;
 
 	// Tracks which clothing LOD each skel mesh LOD corresponds to (LodMap[SkelLod]=ClothingLod).
 	UPROPERTY()
diff --git a/Engine/Source/Runtime/ClothingSystemRuntimeNv/Private/Assets/ClothingAssetNv.cpp b/Engine/Source/Runtime/ClothingSystemRuntimeNv/Private/Assets/ClothingAssetNv.cpp
--- a/Engine/Source/Runtime/ClothingSystemRuntimeNv/Private/Assets/ClothingAssetNv.cpp
+++ b/Engine/Source/Runtime/ClothingSystemRuntimeNv/Private/Assets/ClothingAssetNv.cpp
@@ -97,11 +97,11 @@ void UClothingAssetNv::PostLoad()
 	{
 		// Remap struct FClothConfig to class UClothConfigNv
 		ClothConfig_DEPRECATED.MigrateTo(Cast<UClothConfigNv>(ClothSimConfig));
-		// Remap struct FClothLODData to class UClothLODDataNv
+		// Remap struct FClothLODData to class UClothLODDataCommon
 		for (FClothLODData &FLod : LodData_DEPRECATED)
 		{
 			const int32 Idx = AddNewLod();
-			FLod.MigrateTo(Cast<UClothLODDataNv>(ClothLodData[Idx]));
+			FLod.MigrateTo(ClothLodData[Idx]);
 		}
 		LodData_DEPRECATED.Empty();
 	}
@@ -109,10 +109,10 @@ void UClothingAssetNv::PostLoad()
 	{
 #if WITH_EDITORONLY_DATA
 		// Convert current parameters to masks
-		for(UClothLODDataBase* LodPtr : ClothLodData)
+		for(UClothLODDataCommon* LodPtr : ClothLodData)
 		{
 			check(LodPtr);
-			UClothLODDataNv& Lod = *Cast<UClothLODDataNv>(LodPtr);
+			UClothLODDataCommon& Lod = *LodPtr;
 			check(Lod.PhysicalMeshData);
 			UClothPhysicalMeshDataBase& PhysMesh = *Lod.PhysicalMeshData;
 
@@ -155,7 +155,7 @@ void UClothingAssetNv::PostLoad()
 	// Fix content imported before we kept vertex colors
 	if(ClothingCustomVersion < FClothingAssetCustomVersion::AddVertexColorsToPhysicalMesh)
 	{
-		for (UClothLODDataBase* Lod : ClothLodData)
+		for (UClothLODDataCommon* Lod : ClothLodData)
 		{
 			const int32 NumVerts = Lod->PhysicalMeshData->Vertices.Num(); // number of verts
 
@@ -195,10 +195,10 @@ void UClothingAssetNv::Serialize(FArchive& Ar)
 
 void UClothingAssetNv::InvalidateCachedData()
 {
-	for(UClothLODDataBase* CurrentLodDataPtr : ClothLodData)
+	for(UClothLODDataCommon* CurrentLodDataPtr : ClothLodData)
 	{
 		check(CurrentLodDataPtr);
-		UClothLODDataNv& CurrentLodData = *Cast<UClothLODDataNv>(CurrentLodDataPtr);
+		UClothLODDataCommon& CurrentLodData = *CurrentLodDataPtr;
 		// Recalculate inverse masses for the physical mesh particles
 		check(CurrentLodData.PhysicalMeshData);
 		UClothPhysicalMeshDataBase& PhysMesh = *CurrentLodData.PhysicalMeshData;
diff --git a/Engine/Source/Runtime/ClothingSystemRuntimeNv/Private/ClothLODDataNv.cpp b/Engine/Source/Runtime/ClothingSystemRuntimeNv/Private/ClothLODDataNv.cpp
--- a/Engine/Source/Runtime/ClothingSystemRuntimeNv/Private/ClothLODDataNv.cpp
+++ b/Engine/Source/Runtime/ClothingSystemRuntimeNv/Private/ClothLODDataNv.cpp
@@ -10,7 +10,7 @@ UClothLODDataNv::UClothLODDataNv(const FObjectInitializer& Init)
 UClothLODDataNv::~UClothLODDataNv()
 {}
 
-void FClothLODData::MigrateTo(UClothLODDataNv* LodData) const
+void FClothLODData::MigrateTo(UClothLODDataCommon* LodData) const
 {
 	PhysicalMeshData.MigrateTo(LodData->PhysicalMeshData);
 	LodData->CollisionData = CollisionData;
diff --git a/Engine/Source/Runtime/ClothingSystemRuntimeNv/Private/ClothingSimulationNv.cpp b/Engine/Source/Runtime/ClothingSystemRuntimeNv/Private/ClothingSimulationNv.cpp
--- a/Engine/Source/Runtime/ClothingSystemRuntimeNv/Private/ClothingSimulationNv.cpp
+++ b/Engine/Source/Runtime/ClothingSystemRuntimeNv/Private/ClothingSimulationNv.cpp
@@ -215,7 +215,7 @@ void FClothingSimulationNv::CreateActor(USkeletalMeshComponent* InOwnerComponent
 
 	for(int32 LodIndex = 0; LodIndex < Asset->ClothLodData.Num(); ++LodIndex)
 	{
-		UClothLODDataBase& AssetLodData = *Asset->ClothLodData[LodIndex];
+		UClothLODDataCommon& AssetLodData = *Asset->ClothLodData[LodIndex];
 		UClothPhysicalMeshDataNv* PhysMeshPtr = Cast<UClothPhysicalMeshDataNv>(AssetLodData.PhysicalMeshData);
 		if (!ensure(PhysMeshPtr))
 		{
@@ -1209,7 +1209,7 @@ void FClothingSimulationNv::UpdateLod(int32 InPredictedLod, const FTransform& Co
 				{
 					// Reposition particles skinned to outgoing LOD
 					bool bLodTransitionUp = OldClothingLod < PredictedClothingLod;
-					UClothLODDataBase& NewLodAssetData = *Actor.AssetCreatedFrom->ClothLodData[PredictedClothingLod];
+					UClothLODDataCommon& NewLodAssetData = *Actor.AssetCreatedFrom->ClothLodData[PredictedClothingLod];
 					TArray<FMeshToMeshVertData>& SkinData = bLodTransitionUp ? NewLodAssetData.TransitionUpSkinData : NewLodAssetData.TransitionDownSkinData;
 
 					for(int32 ParticleIndex = 0; ParticleIndex < NumNewParticles; ++ParticleIndex)
@@ -1666,7 +1666,7 @@ void FClothingSimulationNv::DebugDraw_SelfCollision(USkeletalMeshComponent* Owne
 
 		const float SelfCollisionThickness = Config.SelfCollisionRadius;
 
-		const UClothLODDataBase& LodData = *Asset->ClothLodData[Actor.CurrentLodIndex];
+		const UClothLODDataCommon& LodData = *Asset->ClothLodData[Actor.CurrentLodIndex];
 		const UClothPhysicalMeshDataBase& PhysMesh = *LodData.PhysicalMeshData;
 
 		nv::cloth::Cloth* CurrentCloth = Actor.LodData[Actor.CurrentLodIndex].Cloth;
diff --git a/Engine/Source/Runtime/ClothingSystemRuntimeNv/Public/ClothLODDataNv.h b/Engine/Source/Runtime/ClothingSystemRuntimeNv/Public/ClothLODDataNv.h
--- a/Engine/Source/Runtime/ClothingSystemRuntimeNv/Public/ClothLODDataNv.h
+++ b/Engine/Source/Runtime/ClothingSystemRuntimeNv/Public/ClothLODDataNv.h
@@ -15,7 +15,7 @@
 
 /** */
 UCLASS()
-class CLOTHINGSYSTEMRUNTIMENV_API UClothLODDataNv : public UClothLODDataBase
+class CLOTHINGSYSTEMRUNTIMENV_API UClothLODDataNv : public UClothLODDataCommon
 {
 	GENERATED_BODY()
 public:
@@ -65,7 +65,7 @@ struct CLOTHINGSYSTEMRUNTIMENV_API FClothLODData
 	}
 
 	// Migrate this legacy Nv struct to the new class format (called by UClothingAssetNv::Postload())
-	void MigrateTo(UClothLODDataNv* LodData) const;
+	void MigrateTo(UClothLODDataCommon* LodData) const;
 };
 
 template<>
