fix crash when deleting cloth asset from the cloth paint windows dialog

diff --git a/Engine/Source/Editor/ClothPainter/Private/SClothAssetSelector.cpp b/Engine/Source/Editor/ClothPainter/Private/SClothAssetSelector.cpp
--- a/Engine/Source/Editor/ClothPainter/Private/SClothAssetSelector.cpp
+++ b/Engine/Source/Editor/ClothPainter/Private/SClothAssetSelector.cpp
@@ -29,6 +29,7 @@
 #include "SCopyVertexColorSettingsPanel.h"
 #include "Editor/ContentBrowser/Private/SAssetPicker.h"
 #include "ContentBrowserModule.h"
+#include "Animation/DebugSkelMeshComponent.h"
 
 #define LOCTEXT_NAMESPACE "ClothAssetSelector"
 
@@ -187,24 +188,28 @@ private:
 
 	void DeleteAsset()
 	{
+		//Lambda use to sync one of the UserSectionData section from one LOD Model
+		auto SetSkelMeshSourceSectionUserData = [](FSkeletalMeshLODModel& LODModel, const int32 SectionIndex, const int32 OriginalSectionIndex)
+		{
+			FSkelMeshSourceSectionUserData& SourceSectionUserData = LODModel.UserSectionsData.FindOrAdd(OriginalSectionIndex);
+			SourceSectionUserData.bDisabled = LODModel.Sections[SectionIndex].bDisabled;
+			SourceSectionUserData.bCastShadow = LODModel.Sections[SectionIndex].bCastShadow;
+			SourceSectionUserData.bRecomputeTangent = LODModel.Sections[SectionIndex].bRecomputeTangent;
+			SourceSectionUserData.GenerateUpToLodIndex = LODModel.Sections[SectionIndex].GenerateUpToLodIndex;
+			SourceSectionUserData.CorrespondClothAssetIndex = LODModel.Sections[SectionIndex].CorrespondClothAssetIndex;
+			SourceSectionUserData.ClothingData = LODModel.Sections[SectionIndex].ClothingData;
+		};
+
 		if(UClothingAssetCommon* Asset = Item->ClothingAsset.Get())
 		{
 			if(USkeletalMesh* SkelMesh = Cast<USkeletalMesh>(Asset->GetOuter()))
 			{
+				FScopedSuspendAlternateSkinWeightPreview ScopedSuspendAlternateSkinnWeightPreview(SkelMesh);
 				int32 AssetIndex;
 				if(SkelMesh->MeshClothingAssets.Find(Asset, AssetIndex))
 				{
-					TArray<UActorComponent*> ComponentsToReregister;
-					for(TObjectIterator<USkeletalMeshComponent> It; It; ++It)
-					{
-						if((*It)->SkeletalMesh == SkelMesh)
-						{
-							ComponentsToReregister.Add(*It);
-						}
-					}
-
-					FMultiComponentReregisterContext ReregisterContext(ComponentsToReregister);
-
+					// Need to unregister our components so they shut down their current clothing simulation
+					FScopedSkeletalMeshPostEditChange ScopedPostEditChange(SkelMesh);
 					SkelMesh->PreEditChange(nullptr);
 
 					Asset->UnbindFromSkeletalMesh(SkelMesh);
@@ -215,18 +220,18 @@ private:
 					{
 						for(FSkeletalMeshLODModel& LodModel : Model->LODModels)
 						{
-							for(FSkelMeshSection& Section : LodModel.Sections)
+							for (int32 SectionIndex = 0; SectionIndex < LodModel.Sections.Num(); ++SectionIndex)
 							{
+								FSkelMeshSection& Section = LodModel.Sections[SectionIndex];
 								if(Section.CorrespondClothAssetIndex > AssetIndex)
 								{
 									--Section.CorrespondClothAssetIndex;
+									//Keep the user section data (build source data) in sync
+									SetSkelMeshSourceSectionUserData(LodModel, SectionIndex, Section.OriginalDataSectionIndex);
 								}
 							}
 						}
 					}
-
-					SkelMesh->PostEditChange();
-
 					OnInvalidateList.ExecuteIfBound();
 				}
 			}
diff --git a/Engine/Source/Editor/Persona/Private/PersonaMeshDetails.cpp b/Engine/Source/Editor/Persona/Private/PersonaMeshDetails.cpp
--- a/Engine/Source/Editor/Persona/Private/PersonaMeshDetails.cpp
+++ b/Engine/Source/Editor/Persona/Private/PersonaMeshDetails.cpp
@@ -5530,46 +5530,32 @@ FReply FPersonaMeshDetails::OnRemoveClothingAssetClicked(int32 AssetIndex, IDeta
 	USkeletalMesh* SkelMesh = GetPersonaToolkit()->GetMesh();
 	check(SkelMesh);
 
-	TArray<UActorComponent*> ComponentsToReregister;
-	for(TObjectIterator<USkeletalMeshComponent> It; It; ++It)
-	{
-		if(USkeletalMesh* UsedMesh = (*It)->SkeletalMesh)
-		{
-			if(UsedMesh == SkelMesh)
-			{
-				ComponentsToReregister.Add(*It);
-			}
-		}
-	}
-
 	FScopedSuspendAlternateSkinWeightPreview ScopedSuspendAlternateSkinnWeightPreview(SkelMesh);
+	// Now we can remove the asset.
+	if(SkelMesh->MeshClothingAssets.IsValidIndex(AssetIndex))
 	{
 		// Need to unregister our components so they shut down their current clothing simulation
 		FScopedSkeletalMeshPostEditChange ScopedPostEditChange(SkelMesh);
 
-		// Now we can remove the asset.
-		if(SkelMesh->MeshClothingAssets.IsValidIndex(AssetIndex))
+		if (UClothingAssetBase* AssetToRemove = SkelMesh->MeshClothingAssets[AssetIndex])
 		{
-			if (UClothingAssetBase* AssetToRemove = SkelMesh->MeshClothingAssets[AssetIndex])
-			{
-				AssetToRemove->UnbindFromSkeletalMesh(SkelMesh);
-			}
-			SkelMesh->MeshClothingAssets.RemoveAt(AssetIndex);
+			AssetToRemove->UnbindFromSkeletalMesh(SkelMesh);
+		}
+		SkelMesh->MeshClothingAssets.RemoveAt(AssetIndex);
 
-			// Need to fix up asset indices on sections.
-			if(FSkeletalMeshModel* MeshResource = SkelMesh->GetImportedModel())
+		// Need to fix up asset indices on sections.
+		if(FSkeletalMeshModel* MeshResource = SkelMesh->GetImportedModel())
+		{
+			for(FSkeletalMeshLODModel& LodModel : MeshResource->LODModels)
 			{
-				for(FSkeletalMeshLODModel& LodModel : MeshResource->LODModels)
+				for(int32 SectionIndex = 0; SectionIndex < LodModel.Sections.Num(); ++SectionIndex)
 				{
-					for(int32 SectionIndex = 0; SectionIndex < LodModel.Sections.Num(); ++SectionIndex)
+					FSkelMeshSection& Section = LodModel.Sections[SectionIndex];
+					if(Section.CorrespondClothAssetIndex > AssetIndex)
 					{
-						FSkelMeshSection& Section = LodModel.Sections[SectionIndex];
-						if(Section.CorrespondClothAssetIndex > AssetIndex)
-						{
-							--Section.CorrespondClothAssetIndex;
-							//Keep the user section data (build source data) in sync
-							SetSkelMeshSourceSectionUserData(LodModel, SectionIndex, Section.OriginalDataSectionIndex);
-						}
+						--Section.CorrespondClothAssetIndex;
+						//Keep the user section data (build source data) in sync
+						SetSkelMeshSourceSectionUserData(LodModel, SectionIndex, Section.OriginalDataSectionIndex);
 					}
 				}
 			}
