Persona based editors use asset editor toolkits to create and own their mode managers

diff --git a/Engine/Plugins/Experimental/ControlRig/Source/ControlRigEditor/Private/Editor/ControlRigEditor.cpp b/Engine/Plugins/Experimental/ControlRig/Source/ControlRigEditor/Private/Editor/ControlRigEditor.cpp
--- a/Engine/Plugins/Experimental/ControlRig/Source/ControlRigEditor/Private/Editor/ControlRigEditor.cpp
+++ b/Engine/Plugins/Experimental/ControlRig/Source/ControlRigEditor/Private/Editor/ControlRigEditor.cpp
@@ -17,7 +17,7 @@
 #include "PersonaModule.h"
 #include "ControlRigEditorEditMode.h"
 #include "ControlRigEditModeSettings.h"
-#include "AssetEditorModeManager.h"
+#include "EditorModeManager.h"
 #include "ControlRigBlueprintGeneratedClass.h"
 #include "AnimCustomInstanceHelper.h"
 #include "Sequencer/ControlRigLayerInstance.h"
@@ -229,12 +229,8 @@ void FControlRigEditor::InitControlRigEditor(const EToolkitMode::Type Mode, cons
 	SetCurrentMode(FControlRigEditorModes::ControlRigEditorMode);
 
 	// Activate our edit mode
-//	GetAssetEditorModeManager()->SetToolkitHost(GetToolkitHost());
-	if (GetAssetEditorModeManager() != nullptr)
-	{
-		GetAssetEditorModeManager()->SetDefaultMode(FControlRigEditorEditMode::ModeName);
-		GetAssetEditorModeManager()->ActivateMode(FControlRigEditorEditMode::ModeName);
-	}
+	GetEditorModeManager().SetDefaultMode(FControlRigEditorEditMode::ModeName);
+	GetEditorModeManager().ActivateMode(FControlRigEditorEditMode::ModeName);
 
 	if (FControlRigEditMode* EditMode = GetEditMode())
 	{
@@ -2604,11 +2600,7 @@ void FControlRigEditor::OnControlUISettingChanged(FRigHierarchyContainer* Contai
 
 FControlRigEditorEditMode* FControlRigEditor::GetEditMode() const
 {
-	if (GetAssetEditorModeManager() == nullptr)
-	{
-		return nullptr;
-	}
-	return static_cast<FControlRigEditorEditMode*>(GetAssetEditorModeManager()->GetActiveMode(FControlRigEditorEditMode::ModeName));
+	return static_cast<FControlRigEditorEditMode*>(GetEditorModeManager().GetActiveMode(FControlRigEditorEditMode::ModeName));
 }
 
 
diff --git a/Engine/Source/Editor/AnimationBlueprintEditor/Private/AnimationBlueprintEditor.cpp b/Engine/Source/Editor/AnimationBlueprintEditor/Private/AnimationBlueprintEditor.cpp
--- a/Engine/Source/Editor/AnimationBlueprintEditor/Private/AnimationBlueprintEditor.cpp
+++ b/Engine/Source/Editor/AnimationBlueprintEditor/Private/AnimationBlueprintEditor.cpp
@@ -1429,6 +1429,15 @@ void FAnimationBlueprintEditor::OnBlueprintChangedImpl(UBlueprint* InBlueprint,
 	OnPostCompile();
 }
 
+void FAnimationBlueprintEditor::CreateEditorModeManager()
+{
+	TSharedPtr<FAssetEditorModeManager> ModeManager = MakeShareable(FModuleManager::LoadModuleChecked<FPersonaModule>("Persona").CreatePersonaEditorModeManager());
+
+	PRAGMA_DISABLE_DEPRECATION_WARNINGS
+	SetAssetEditorModeManager(ModeManager.Get());
+	PRAGMA_ENABLE_DEPRECATION_WARNINGS
+}
+
 TSharedRef<IPersonaPreviewScene> FAnimationBlueprintEditor::GetPreviewScene() const
 { 
 	return PersonaToolkit->GetPreviewScene(); 
@@ -1485,7 +1494,7 @@ void FAnimationBlueprintEditor::OnSelectedNodesChangedImpl(const TSet<class UObj
 {
 	FBlueprintEditor::OnSelectedNodesChangedImpl(NewSelection);
 
-	IPersonaEditorModeManager* PersonaEditorModeManager = static_cast<IPersonaEditorModeManager*>(GetAssetEditorModeManager());
+	IPersonaEditorModeManager* PersonaEditorModeManager = static_cast<IPersonaEditorModeManager*>(&GetEditorModeManager());
 
 	if (UAnimGraphNode_Base* SelectedAnimGraphNodePtr = SelectedAnimGraphNode.Get())
 	{
diff --git a/Engine/Source/Editor/AnimationBlueprintEditor/Private/AnimationBlueprintEditor.h b/Engine/Source/Editor/AnimationBlueprintEditor/Private/AnimationBlueprintEditor.h
--- a/Engine/Source/Editor/AnimationBlueprintEditor/Private/AnimationBlueprintEditor.h
+++ b/Engine/Source/Editor/AnimationBlueprintEditor/Private/AnimationBlueprintEditor.h
@@ -194,7 +194,8 @@ protected:
 	virtual FGraphAppearanceInfo GetGraphAppearance(class UEdGraph* InGraph) const override;
 	virtual bool IsEditable(UEdGraph* InGraph) const override;
 	virtual FText GetGraphDecorationString(UEdGraph* InGraph) const override;
-	virtual void OnBlueprintChangedImpl(UBlueprint* InBlueprint, bool bIsJustBeingCompiled = false) override;	
+	virtual void OnBlueprintChangedImpl(UBlueprint* InBlueprint, bool bIsJustBeingCompiled = false) override;
+	virtual void CreateEditorModeManager() override;
 	//~ End FBlueprintEditor Interface
 
 	//~ Begin FEditorUndoClient Interface
diff --git a/Engine/Source/Editor/AnimationEditor/Private/AnimationEditor.cpp b/Engine/Source/Editor/AnimationEditor/Private/AnimationEditor.cpp
--- a/Engine/Source/Editor/AnimationEditor/Private/AnimationEditor.cpp
+++ b/Engine/Source/Editor/AnimationEditor/Private/AnimationEditor.cpp
@@ -50,6 +50,8 @@
 #include "Algo/Transform.h"
 #include "ISequenceRecorder.h"
 #include "IAnimSequenceCurveEditor.h"
+#include "EditorModeManager.h"
+#include "IPersonaEditorModeManager.h"
 
 const FName AnimationEditorAppIdentifier = FName(TEXT("AnimationEditorApp"));
 
diff --git a/Engine/Source/Editor/AnimationEditor/Public/IAnimationEditor.h b/Engine/Source/Editor/AnimationEditor/Public/IAnimationEditor.h
--- a/Engine/Source/Editor/AnimationEditor/Public/IAnimationEditor.h
+++ b/Engine/Source/Editor/AnimationEditor/Public/IAnimationEditor.h
@@ -3,7 +3,7 @@
 #pragma once
 
 #include "CoreMinimal.h"
-#include "WorkflowOrientedApp/WorkflowCentricApplication.h"
+#include "PersonaAssetEditorToolkit.h"
 #include "IHasPersonaToolkit.h"
 #include "Animation/SmartName.h"
 
@@ -14,7 +14,7 @@ struct FRichCurve;
 class ITimeSliderController;
 enum class ERawCurveTrackTypes : uint8;
 
-class IAnimationEditor : public FWorkflowCentricApplication, public IHasPersonaToolkit
+class IAnimationEditor : public FPersonaAssetEditorToolkit, public IHasPersonaToolkit
 {
 public:
 	/** Set the animation asset of the editor. */
diff --git a/Engine/Source/Editor/ClothPainter/Private/SClothPaintTab.cpp b/Engine/Source/Editor/ClothPainter/Private/SClothPaintTab.cpp
--- a/Engine/Source/Editor/ClothPainter/Private/SClothPaintTab.cpp
+++ b/Engine/Source/Editor/ClothPainter/Private/SClothPaintTab.cpp
@@ -42,10 +42,7 @@ SClothPaintTab::~SClothPaintTab()
 {
 	if(ISkeletalMeshEditor* SkeletalMeshEditor = static_cast<ISkeletalMeshEditor*>(HostingApp.Pin().Get()))
 	{
-		if(FAssetEditorModeManager* ModeManager = SkeletalMeshEditor->GetAssetEditorModeManager())
-		{
-			ModeManager->ActivateDefaultMode();
-		}
+		SkeletalMeshEditor->GetEditorModeManager().ActivateDefaultMode();
 	}
 }
 
@@ -145,9 +142,9 @@ void SClothPaintTab::UpdatePaintTools()
 	if (bPaintModeEnabled)
 	{
 		ISkeletalMeshEditor* SkeletalMeshEditor = static_cast<ISkeletalMeshEditor*>(HostingApp.Pin().Get());
-		SkeletalMeshEditor->GetAssetEditorModeManager()->ActivateMode(PaintModeID, true);
+		SkeletalMeshEditor->GetEditorModeManager().ActivateMode(PaintModeID, true);
 
-		FClothingPaintEditMode* PaintMode = (FClothingPaintEditMode*)SkeletalMeshEditor->GetAssetEditorModeManager()->GetActiveMode(PaintModeID);
+		FClothingPaintEditMode* PaintMode = (FClothingPaintEditMode*)SkeletalMeshEditor->GetEditorModeManager().GetActiveMode(PaintModeID);
 		if (PaintMode)
 		{
 			FClothPainter* ClothPainter = static_cast<FClothPainter*>(PaintMode->GetMeshPainter());
@@ -178,7 +175,7 @@ void SClothPaintTab::UpdatePaintTools()
 	{
 		ContentBox->RemoveSlot(ModeWidget->AsShared());
 		ISkeletalMeshEditor* SkeletalMeshEditor = static_cast<ISkeletalMeshEditor*>(HostingApp.Pin().Get());
-		SkeletalMeshEditor->GetAssetEditorModeManager()->ActivateDefaultMode();
+		SkeletalMeshEditor->GetEditorModeManager().ActivateDefaultMode();
 		ModeWidget = nullptr;
 	}
 }
@@ -189,7 +186,7 @@ void SClothPaintTab::OnAssetSelectionChanged(TWeakObjectPtr<UClothingAssetCommon
 	{
 		ISkeletalMeshEditor* SkeletalMeshEditor = static_cast<ISkeletalMeshEditor*>(HostingApp.Pin().Get());
 
-		FClothingPaintEditMode* PaintMode = (FClothingPaintEditMode*)SkeletalMeshEditor->GetAssetEditorModeManager()->GetActiveMode(PaintModeID);
+		FClothingPaintEditMode* PaintMode = (FClothingPaintEditMode*)SkeletalMeshEditor->GetEditorModeManager().GetActiveMode(PaintModeID);
 		if(PaintMode)
 		{
 			FClothPainter* ClothPainter = static_cast<FClothPainter*>(PaintMode->GetMeshPainter());
diff --git a/Engine/Source/Editor/Persona/Private/AnimationEditorViewportClient.cpp b/Engine/Source/Editor/Persona/Private/AnimationEditorViewportClient.cpp
--- a/Engine/Source/Editor/Persona/Private/AnimationEditorViewportClient.cpp
+++ b/Engine/Source/Editor/Persona/Private/AnimationEditorViewportClient.cpp
@@ -61,7 +61,7 @@ IMPLEMENT_HIT_PROXY( HPersonaBoneProxy, HHitProxy );
 // FAnimationViewportClient
 
 FAnimationViewportClient::FAnimationViewportClient(const TSharedRef<IPersonaPreviewScene>& InPreviewScene, const TSharedRef<SAnimationEditorViewport>& InAnimationEditorViewport, const TSharedRef<FAssetEditorToolkit>& InAssetEditorToolkit, int32 InViewportIndex, bool bInShowStats)
-	: FEditorViewportClient(FModuleManager::LoadModuleChecked<FPersonaModule>("Persona").CreatePersonaEditorModeManager(), &InPreviewScene.Get(), StaticCastSharedRef<SEditorViewport>(InAnimationEditorViewport))
+	: FEditorViewportClient(&InAssetEditorToolkit->GetEditorModeManager(), &InPreviewScene.Get(), StaticCastSharedRef<SEditorViewport>(InAnimationEditorViewport))
 	, PreviewScenePtr(InPreviewScene)
 	, AssetEditorToolkitPtr(InAssetEditorToolkit)
 	, AnimationPlaybackSpeedMode(EAnimationPlaybackSpeeds::Normal)
@@ -79,14 +79,9 @@ FAnimationViewportClient::FAnimationViewportClient(const TSharedRef<IPersonaPrev
 
 	InPreviewScene->RegisterOnCameraOverrideChanged(FSimpleDelegate::CreateRaw(this, &FAnimationViewportClient::OnCameraControllerChanged));
 
-	// Give ownership of the mode manager to the asset editor toolkit so it can be used outside of thew viewport
-	PRAGMA_DISABLE_DEPRECATION_WARNINGS
-	InAssetEditorToolkit->SetAssetEditorModeManager((FAssetEditorModeManager*)ModeTools);
-	PRAGMA_ENABLE_DEPRECATION_WARNINGS
-
 	Widget->SetUsesEditorModeTools(ModeTools);
 	((FAssetEditorModeManager*)ModeTools)->SetPreviewScene(&InPreviewScene.Get());
-	((FAssetEditorModeManager*)ModeTools)->SetDefaultMode(FPersonaEditModes::SkeletonSelection);
+	ModeTools->SetDefaultMode(FPersonaEditModes::SkeletonSelection);
 
 	// Default to local space
 	SetWidgetCoordSystemSpace(COORD_Local);
diff --git a/Engine/Source/Editor/Persona/Private/PersonaAssetEditorToolkit.cpp b/Engine/Source/Editor/Persona/Private/PersonaAssetEditorToolkit.cpp
new file mode 100644
--- /dev/null
+++ b/Engine/Source/Editor/Persona/Private/PersonaAssetEditorToolkit.cpp
@@ -0,0 +1,17 @@
+// Copyright Epic Games, Inc. All Rights Reserved.
+
+#include "PersonaAssetEditorToolkit.h"
+
+#include "PersonaModule.h"
+#include "AssetEditorModeManager.h"
+#include "Modules/ModuleManager.h"
+#include "IPersonaEditorModeManager.h"
+
+void FPersonaAssetEditorToolkit::CreateEditorModeManager()
+{
+	TSharedPtr<FAssetEditorModeManager> ModeManager = MakeShareable(FModuleManager::LoadModuleChecked<FPersonaModule>("Persona").CreatePersonaEditorModeManager());
+	
+	PRAGMA_DISABLE_DEPRECATION_WARNINGS
+	SetAssetEditorModeManager(ModeManager.Get());
+	PRAGMA_ENABLE_DEPRECATION_WARNINGS
+}
diff --git a/Engine/Source/Editor/Persona/Public/PersonaAssetEditorToolkit.h b/Engine/Source/Editor/Persona/Public/PersonaAssetEditorToolkit.h
new file mode 100644
--- /dev/null
+++ b/Engine/Source/Editor/Persona/Public/PersonaAssetEditorToolkit.h
@@ -0,0 +1,13 @@
+// Copyright Epic Games, Inc. All Rights Reserved.
+
+#pragma once
+
+#include "WorkflowOrientedApp/WorkflowCentricApplication.h"
+
+/** Persona asset editor toolkit wrapper, used to auto inject the persona editor mode manager */
+class PERSONA_API FPersonaAssetEditorToolkit : public FWorkflowCentricApplication
+{
+public:
+	/** FAssetEditorToolkit interface  */
+	virtual void CreateEditorModeManager() override;
+};
diff --git a/Engine/Source/Editor/PhysicsAssetEditor/Private/PhysicsAssetEditor.cpp b/Engine/Source/Editor/PhysicsAssetEditor/Private/PhysicsAssetEditor.cpp
--- a/Engine/Source/Editor/PhysicsAssetEditor/Private/PhysicsAssetEditor.cpp
+++ b/Engine/Source/Editor/PhysicsAssetEditor/Private/PhysicsAssetEditor.cpp
@@ -181,10 +181,10 @@ void FPhysicsAssetEditor::InitPhysicsAssetEditor(const EToolkitMode::Type Mode,
 	// Force disable simulation as InitArticulated can be called during viewport creation
 	SharedData->ForceDisableSimulation();
 
-	GetAssetEditorModeManager()->SetDefaultMode(FPhysicsAssetEditorEditMode::ModeName);
-	GetAssetEditorModeManager()->ActivateMode(FPersonaEditModes::SkeletonSelection);
-	GetAssetEditorModeManager()->ActivateMode(FPhysicsAssetEditorEditMode::ModeName);
-	static_cast<FPhysicsAssetEditorEditMode*>(GetAssetEditorModeManager()->GetActiveMode(FPhysicsAssetEditorEditMode::ModeName))->SetSharedData(SharedThis(this), *SharedData.Get());
+	GetEditorModeManager().SetDefaultMode(FPhysicsAssetEditorEditMode::ModeName);
+	GetEditorModeManager().ActivateMode(FPersonaEditModes::SkeletonSelection);
+	GetEditorModeManager().ActivateMode(FPhysicsAssetEditorEditMode::ModeName);
+	static_cast<FPhysicsAssetEditorEditMode*>(GetEditorModeManager().GetActiveMode(FPhysicsAssetEditorEditMode::ModeName))->SetSharedData(SharedThis(this), *SharedData.Get());
 
 	IPhysicsAssetEditorModule* PhysicsAssetEditorModule = &FModuleManager::LoadModuleChecked<IPhysicsAssetEditorModule>( "PhysicsAssetEditor" );
 	ExtendMenu();
diff --git a/Engine/Source/Editor/PhysicsAssetEditor/Public/IPhysicsAssetEditor.h b/Engine/Source/Editor/PhysicsAssetEditor/Public/IPhysicsAssetEditor.h
--- a/Engine/Source/Editor/PhysicsAssetEditor/Public/IPhysicsAssetEditor.h
+++ b/Engine/Source/Editor/PhysicsAssetEditor/Public/IPhysicsAssetEditor.h
@@ -3,7 +3,7 @@
 #pragma once
 
 #include "CoreMinimal.h"
-#include "WorkflowOrientedApp/WorkflowCentricApplication.h"
+#include "PersonaAssetEditorToolkit.h"
 #include "IHasPersonaToolkit.h"
 
 
@@ -11,7 +11,7 @@
    IPhysicsAssetEditor
 -----------------------------------------------------------------------------*/
 
-class IPhysicsAssetEditor : public FWorkflowCentricApplication, public IHasPersonaToolkit
+class IPhysicsAssetEditor : public FPersonaAssetEditorToolkit, public IHasPersonaToolkit
 {
 
 };
diff --git a/Engine/Source/Editor/SkeletalMeshEditor/Public/ISkeletalMeshEditor.h b/Engine/Source/Editor/SkeletalMeshEditor/Public/ISkeletalMeshEditor.h
--- a/Engine/Source/Editor/SkeletalMeshEditor/Public/ISkeletalMeshEditor.h
+++ b/Engine/Source/Editor/SkeletalMeshEditor/Public/ISkeletalMeshEditor.h
@@ -3,10 +3,10 @@
 #pragma once
 
 #include "CoreMinimal.h"
-#include "WorkflowOrientedApp/WorkflowCentricApplication.h"
+#include "PersonaAssetEditorToolkit.h"
 #include "IHasPersonaToolkit.h"
 
-class ISkeletalMeshEditor : public FWorkflowCentricApplication, public IHasPersonaToolkit
+class ISkeletalMeshEditor : public FPersonaAssetEditorToolkit, public IHasPersonaToolkit
 {
 };
 
diff --git a/Engine/Source/Editor/SkeletonEditor/Public/ISkeletonEditor.h b/Engine/Source/Editor/SkeletonEditor/Public/ISkeletonEditor.h
--- a/Engine/Source/Editor/SkeletonEditor/Public/ISkeletonEditor.h
+++ b/Engine/Source/Editor/SkeletonEditor/Public/ISkeletonEditor.h
@@ -3,9 +3,9 @@
 #pragma once
 
 #include "CoreMinimal.h"
-#include "WorkflowOrientedApp/WorkflowCentricApplication.h"
+#include "PersonaAssetEditorToolkit.h"
 #include "IHasPersonaToolkit.h"
 
-class ISkeletonEditor : public FWorkflowCentricApplication, public IHasPersonaToolkit
+class ISkeletonEditor : public FPersonaAssetEditorToolkit, public IHasPersonaToolkit
 {
 };
