Cloth - Fix cloth painting lag after each strokes.

Merging CL 14009849

diff --git a/Engine/Source/Editor/ClothPainter/Private/ClothPainter.cpp b/Engine/Source/Editor/ClothPainter/Private/ClothPainter.cpp
--- a/Engine/Source/Editor/ClothPainter/Private/ClothPainter.cpp
+++ b/Engine/Source/Editor/ClothPainter/Private/ClothPainter.cpp
@@ -297,19 +297,6 @@ void FClothPainter::FinishPainting()
 		EndTransaction();
 		Adapter->PostEdit();
 
-		if(SkeletalMeshComponent)
-		{
-			FComponentReregisterContext ReregisterContext(SkeletalMeshComponent);
-
-			if(USkeletalMesh* SkelMesh = SkeletalMeshComponent->SkeletalMesh)
-			{
-				for(UClothingAssetBase* AssetBase : SkelMesh->MeshClothingAssets)
-				{
-					AssetBase->InvalidateCachedData();
-				}
-			}
-		}
-
 		/** If necessary, recalculate view ranges when set to auto mode */
 		RecalculateAutoViewRange();
 	}
