Skeletal Mesh Cloth Panel Reskin

diff --git a/Engine/Source/Editor/ClothPainter/Private/SClothAssetSelector.cpp b/Engine/Source/Editor/ClothPainter/Private/SClothAssetSelector.cpp
--- a/Engine/Source/Editor/ClothPainter/Private/SClothAssetSelector.cpp
+++ b/Engine/Source/Editor/ClothPainter/Private/SClothAssetSelector.cpp
@@ -15,6 +15,7 @@
 #include "Framework/MultiBox/MultiBoxBuilder.h"
 #include "Widgets/Input/SCheckBox.h"
 #include "Widgets/Text/SInlineEditableTextBlock.h"
+#include "SEditorHeaderButton.h"
 #include "Engine/SkeletalMesh.h"
 #include "ApexClothingUtils.h"
 #include "UObject/UObjectIterator.h"
@@ -734,24 +735,25 @@ void SClothAssetSelector::Construct(const FArguments& InArgs, USkeletalMesh* InM
 	ChildSlot
 	[
 		SNew(SVerticalBox)
+
 		+SVerticalBox::Slot()
-		.Padding(0.0f, 0.0f, 0.0f, 2.0f)
+		.Padding(FMargin(0.0f, 0.0f, 0.0f, 1.0f))
 		.AutoHeight()
 		[
 			SNew(SExpandableArea)
-			.BorderBackgroundColor(FLinearColor(0.6f, 0.6f, 0.6f))
-			.BodyBorderImage(FEditorStyle::GetBrush("DetailsView.CategoryMiddle"))
-			.BodyBorderBackgroundColor(FLinearColor(1.0f, 1.0f, 1.0f))
+			.BorderImage(FAppStyle::Get().GetBrush("DetailsView.CategoryTop"))
 			.HeaderContent()
 			[
 				SAssignNew(AssetHeaderBox, SHorizontalBox)
 				+SHorizontalBox::Slot()
 				.VAlign(VAlign_Center)
+				.Padding(10.f, 0.f, 0.f, 0.f)
 				[
 					SNew(STextBlock)
 					.Text(LOCTEXT("AssetExpander_Title", "Clothing Data"))
-					.Font(FEditorStyle::GetFontStyle("DetailsView.CategoryFontStyle"))
-					.ShadowOffset(FVector2D(1.0f, 1.0f))
+					.TransformPolicy(ETextTransformPolicy::ToUpper)
+					.TextStyle(FAppStyle::Get(), "DetailsView.CategoryTextStyle")
+					.Font(FAppStyle::Get().GetFontStyle("PropertyWindow.BoldFont"))
 				]
 #if WITH_APEX_CLOTHING
 				+SHorizontalBox::Slot()
@@ -761,9 +763,6 @@ void SClothAssetSelector::Construct(const FArguments& InArgs, USkeletalMesh* InM
 				.Padding(0.0f, 0.0f, 4.0f, 0.0f)
 				[
 					SNew(SButton)
-					.ButtonStyle(FEditorStyle::Get(), "RoundButton")
-					.ForegroundColor(FEditorStyle::GetSlateColor("DefaultForeground"))
-					.ContentPadding(FMargin(2, 0))
 					.OnClicked(this, &SClothAssetSelector::OnImportApexFileClicked)
 					.HAlign(HAlign_Center)
 					.VAlign(VAlign_Center)
@@ -771,6 +770,7 @@ void SClothAssetSelector::Construct(const FArguments& InArgs, USkeletalMesh* InM
 						SNew(SHorizontalBox)
 
 						+ SHorizontalBox::Slot()
+						.VAlign(VAlign_Center)
 						.AutoWidth()
 						.Padding(FMargin(0, 1))
 						[
@@ -787,100 +787,74 @@ void SClothAssetSelector::Construct(const FArguments& InArgs, USkeletalMesh* InM
 							.Font(IDetailLayoutBuilder::GetDetailFontBold())
 							.Text(LOCTEXT("NewAssetButtonText", "Import APEX file"))
 							.Visibility(this, &SClothAssetSelector::GetAssetHeaderButtonTextVisibility)
-							.ShadowOffset(FVector2D(1, 1))
 						]
 					]
 				]
 #endif  // #if WITH_APEX_CLOTHING
 				+ SHorizontalBox::Slot()
-					.AutoWidth()
-					.VAlign(VAlign_Center)
-					.HAlign(HAlign_Right)
-					.Padding(0.0f, 0.0f, 4.0f, 0.0f)
-					[
-						SNew(SComboButton)
-						.ButtonStyle(FEditorStyle::Get(), "RoundButton")
-					.ForegroundColor(FEditorStyle::GetSlateColor("DefaultForeground"))
-					.ContentPadding(FMargin(2, 0))
+				.AutoWidth()
+				.VAlign(VAlign_Center)
+				.HAlign(HAlign_Right)
+				.Padding(0.0f, 0.0f, 4.0f, 0.0f)
+				[
+					SNew(SEditorHeaderButton)
+					.Icon(FAppStyle::Get().GetBrush("Icons.Plus"))
+					.ToolTipText(LOCTEXT("CopyClothingFromMeshText_TEXT", "Copy Clothing from SkeletalMesh"))
 					.OnGetMenuContent(this, &SClothAssetSelector::OnGenerateSkeletalMeshPickerForClothCopy)
-					.HAlign(HAlign_Center)
-					.VAlign(VAlign_Center)
-					.ButtonContent()
-					[
-						SNew(SHorizontalBox)
+				]
 
-						+ SHorizontalBox::Slot()
-					.VAlign(VAlign_Center)
-					.AutoWidth()
-					.Padding(FMargin(2, 0, 0, 0))
-					[
-						SNew(STextBlock)
-						.Font(IDetailLayoutBuilder::GetDetailFontBold())
-					.Text(LOCTEXT("CopyClothingFromMeshText", "Copy Clothing from SkeletalMesh"))
-					.ShadowOffset(FVector2D(1, 1))
-					]
-					]
-					]
 				+ SHorizontalBox::Slot()
 				.AutoWidth()
 				.VAlign(VAlign_Center)
 				.HAlign(HAlign_Right)
 				[
 					SNew(SComboButton)
-					.ButtonStyle(FEditorStyle::Get(), "RoundButton")
-					.ForegroundColor(FEditorStyle::GetSlateColor("DefaultForeground"))
-					.ContentPadding(FMargin(2, 0))
+					.ForegroundColor(FSlateColor::UseStyle())
 					.OnGetMenuContent(this, &SClothAssetSelector::OnGetLodMenu)
 					.HasDownArrow(true)
 					.ButtonContent()
 					[
 						SNew(STextBlock)
-						.Font(IDetailLayoutBuilder::GetDetailFontBold())
 						.Text(this, &SClothAssetSelector::GetLodButtonText)
-						.ShadowOffset(FVector2D(1, 1))
 					]
 				]
 			]
 			.BodyContent()
 			[
-				SNew(SVerticalBox)
-				+ SVerticalBox::Slot()
-				.Padding(3.0f)
-				.AutoHeight()
+				SNew(SBox)
+				.Padding(FMargin(0.f, 2.0f))
+				.MinDesiredHeight(100.0f)
 				[
-					SNew(SBox)
-					.MinDesiredHeight(100.0f)
-					.MaxDesiredHeight(100.0f)
-					[
-						SAssignNew(AssetList, SAssetList)
-						.ItemHeight(24)
-						.ListItemsSource(&AssetListItems)
-						.OnGenerateRow(this, &SClothAssetSelector::OnGenerateWidgetForClothingAssetItem)
-						.OnSelectionChanged(this, &SClothAssetSelector::OnAssetListSelectionChanged)
-						.ClearSelectionOnClick(false)
-						.SelectionMode(ESelectionMode::Single)
-					]
+					SAssignNew(AssetList, SAssetList)
+					.ItemHeight(24)
+					.ListItemsSource(&AssetListItems)
+					.OnGenerateRow(this, &SClothAssetSelector::OnGenerateWidgetForClothingAssetItem)
+					.OnSelectionChanged(this, &SClothAssetSelector::OnAssetListSelectionChanged)
+					.ClearSelectionOnClick(false)
+					.SelectionMode(ESelectionMode::Single)
 				]
 			]
 		]
+
 		+ SVerticalBox::Slot()
-		.Padding(0.0f, 0.0f, 0.0f, 2.0f)
+		.Padding(FMargin(0.0f, 0.0f, 0.0f, 1.0f))
 		.AutoHeight()
 		[
 			SNew(SExpandableArea)
-			.BorderBackgroundColor(FLinearColor(0.6f, 0.6f, 0.6f))
-			.BodyBorderImage(FEditorStyle::GetBrush("DetailsView.CategoryMiddle"))
-			.BodyBorderBackgroundColor(FLinearColor(1.0f, 1.0f, 1.0f))
+			.BorderImage(FAppStyle::Get().GetBrush("DetailsView.CategoryTop"))
+			.Padding(FMargin(0.f, 2.0f))
 			.HeaderContent()
 			[
 				SAssignNew(MaskHeaderBox, SHorizontalBox)
 				+SHorizontalBox::Slot()
 				.VAlign(VAlign_Center)
+				.Padding(10.f, 0.f, 0.f, 0.f)
 				[
 					SNew(STextBlock)
 					.Text(LOCTEXT("MaskExpander_Title", "Masks"))
-					.Font(FEditorStyle::GetFontStyle("DetailsView.CategoryFontStyle"))
-					.ShadowOffset(FVector2D(1.0f, 1.0f))
+					.TransformPolicy(ETextTransformPolicy::ToUpper)
+					.TextStyle(FAppStyle::Get(), "DetailsView.CategoryTextStyle")
+					.Font(FAppStyle::Get().GetFontStyle("PropertyWindow.BoldFont"))
 				]
 				+SHorizontalBox::Slot()
 				.AutoWidth()
@@ -888,35 +862,16 @@ void SClothAssetSelector::Construct(const FArguments& InArgs, USkeletalMesh* InM
 				.HAlign(HAlign_Right)
 				[
 					SAssignNew(NewMaskButton, SButton)
-					.ButtonStyle(FEditorStyle::Get(), "RoundButton")
-					.ForegroundColor(FEditorStyle::GetSlateColor("DefaultForeground"))
-					.ContentPadding(FMargin(2, 0))
+					.ButtonStyle( FAppStyle::Get(), "SimpleButton" )
 					.OnClicked(this, &SClothAssetSelector::AddNewMask)
 					.IsEnabled(this, &SClothAssetSelector::CanAddNewMask)
+					.ToolTipText(LOCTEXT("AddMask_Tooltip", "Add a Mask"))
 					.HAlign(HAlign_Center)
 					.VAlign(VAlign_Center)
 					[
-						SNew(SHorizontalBox)
-
-						+ SHorizontalBox::Slot()
-						.AutoWidth()
-						.Padding(FMargin(0, 1))
-						[
-							SNew(SImage)
-							.Image(FEditorStyle::GetBrush("Plus"))
-						]
-
-						+ SHorizontalBox::Slot()
-						.VAlign(VAlign_Center)
-						.AutoWidth()
-						.Padding(FMargin(2, 0, 0, 0))
-						[
-							SNew(STextBlock)
-							.Font(IDetailLayoutBuilder::GetDetailFontBold())
-							.Text(LOCTEXT("NewMaskButtonText", "Mask"))
-							.Visibility(this, &SClothAssetSelector::GetMaskHeaderButtonTextVisibility)
-							.ShadowOffset(FVector2D(1, 1))
-						]
+						SNew(SImage)
+						.ColorAndOpacity(FSlateColor::UseForeground())
+						.Image(FAppStyle::Get().GetBrush("Icons.Plus"))
 					]
 				]
 			]
@@ -924,8 +879,6 @@ void SClothAssetSelector::Construct(const FArguments& InArgs, USkeletalMesh* InM
 			[
 				SNew(SBox)
 				.MinDesiredHeight(100.0f)
-				.MaxDesiredHeight(200.0f)
-				.Padding(FMargin(3.0f))
 				[
 					SAssignNew(MaskList, SMaskList)
 					.ItemHeight(24)
@@ -937,16 +890,15 @@ void SClothAssetSelector::Construct(const FArguments& InArgs, USkeletalMesh* InM
 					.HeaderRow
 					(
 						SNew(SHeaderRow)
+
 						+ SHeaderRow::Column(SMaskListRow::Column_Enabled)
 						.FixedWidth(24)
-						.DefaultLabel(LOCTEXT("MaskListHeader_Enabled", "Enabled"))
-						.HeaderContent()
-						[
-							SNew(SBox)
-						]
+						.DefaultLabel(FText::GetEmpty())
+
 						+ SHeaderRow::Column(SMaskListRow::Column_MaskName)
 						.FillWidth(0.5f)
 						.DefaultLabel(LOCTEXT("MaskListHeader_Name", "Name"))
+
 						+ SHeaderRow::Column(SMaskListRow::Column_CurrentTarget)
 						.FillWidth(0.3f)
 						.DefaultLabel(LOCTEXT("MaskListHeader_Target", "Target"))
@@ -955,23 +907,24 @@ void SClothAssetSelector::Construct(const FArguments& InArgs, USkeletalMesh* InM
 			]
 		]
 		+ SVerticalBox::Slot()					// Mesh to mesh skinning
-		.Padding(0.0f, 0.0f, 0.0f, 2.0f)
 		.AutoHeight()
+		.Padding(FMargin(0.0f, 0.0f, 0.0f, 1.0f))
 		[
 			SNew(SExpandableArea)
-			.BorderBackgroundColor(FLinearColor(0.6f, 0.6f, 0.6f))
-			.BodyBorderImage(FEditorStyle::GetBrush("DetailsView.CategoryMiddle"))
-			.BodyBorderBackgroundColor(FLinearColor(1.0f, 1.0f, 1.0f))
+			.BorderImage(FAppStyle::Get().GetBrush("DetailsView.CategoryTop"))
+			.Padding(FMargin(0.f, 2.0f))
 			.HeaderContent()
 			[
 				SNew(SHorizontalBox)
 				+ SHorizontalBox::Slot()
 				.VAlign(VAlign_Center)
+				.Padding(10, 8, 0, 8)
 				[
 					SNew(STextBlock)
 					.Text(LOCTEXT("MeshSkinning_Title", "Mesh Skinning"))
-					.Font(FEditorStyle::GetFontStyle("DetailsView.CategoryFontStyle"))
-					.ShadowOffset(FVector2D(1.0f, 1.0f))
+					.TransformPolicy(ETextTransformPolicy::ToUpper)
+					.TextStyle(FAppStyle::Get(), "DetailsView.CategoryTextStyle")
+					.Font(FAppStyle::Get().GetFontStyle("PropertyWindow.BoldFont"))
 				]
 			]
 			.BodyContent()
@@ -980,38 +933,39 @@ void SClothAssetSelector::Construct(const FArguments& InArgs, USkeletalMesh* InM
 
 				SNew(SSplitter)
 				.Orientation(EOrientation::Orient_Horizontal)
+				.PhysicalSplitterHandleSize(1.0)
 
 				+ SSplitter::Slot()
 				.Value(.3f)
 				[
 					SNew(SVerticalBox)
+
 					+ SVerticalBox::Slot()
-					.Padding(0.0f, 2.0f, 0.0f, 2.0f)
-					.AutoHeight()
+					.VAlign(VAlign_Center)
+					.Padding(16.0f, 2.0f, 0.0f, 2.0f)
 					[
 						SNew(STextBlock)
 						.Font(IDetailLayoutBuilder::GetDetailFontBold())
 						.Text(LOCTEXT("MultipleInfluences", "Use Multiple Influences"))
-						.ShadowOffset(FVector2D(1, 1))
 					]
 
 					+ SVerticalBox::Slot()
-					.Padding(0.0f, 2.0f, 0.0f, 2.0f)
-					.AutoHeight()
+					.VAlign(VAlign_Center)
+					.Padding(16.0f, 2.0f, 0.0f, 2.0f)
 					[
 						SNew(STextBlock)
 						.Font(IDetailLayoutBuilder::GetDetailFontBold())
 						.Text(LOCTEXT("CurrentRadius", "Kernel Radius"))
-						.ShadowOffset(FVector2D(1, 1))
 					]
 				]
 
 				+ SSplitter::Slot()
 				[
 					SNew(SVerticalBox)
+
 					+ SVerticalBox::Slot()
-					.Padding(0.0f, 2.0f, 0.0f, 2.0f)
-					.AutoHeight()
+					.VAlign(VAlign_Center)
+					.Padding(12.0f, 4.0f, 0.0f, 4.0f)
 					[
 						SNew(SCheckBox)
 						.IsChecked(this, &SClothAssetSelector::GetCurrentUseMultipleInfluences)
@@ -1020,8 +974,8 @@ void SClothAssetSelector::Construct(const FArguments& InArgs, USkeletalMesh* InM
 					]
 
 					+ SVerticalBox::Slot()
-					.Padding(0.0f, 2.0f, 0.0f, 2.0f)
-					.AutoHeight()
+					.VAlign(VAlign_Center)
+					.Padding(12.0f, 4.0f, 0.0f, 4.0f)
 					[
 						SNew(SNumericEntryBox<float>)
 						.AllowSpin(true)
@@ -1039,7 +993,6 @@ void SClothAssetSelector::Construct(const FArguments& InArgs, USkeletalMesh* InM
 
 			]
 		]
-		
 	];
 
 	RefreshAssetList();
diff --git a/Engine/Source/Editor/ClothPainter/Private/SClothPaintTab.cpp b/Engine/Source/Editor/ClothPainter/Private/SClothPaintTab.cpp
--- a/Engine/Source/Editor/ClothPainter/Private/SClothPaintTab.cpp
+++ b/Engine/Source/Editor/ClothPainter/Private/SClothPaintTab.cpp
@@ -77,15 +77,9 @@ void SClothPaintTab::Construct(const FArguments& InArgs)
 	FSlateIcon TexturePaintIcon(FEditorStyle::GetStyleSetName(), "LevelEditor.MeshPaintMode.TexturePaint");
 
 	this->ChildSlot
+	.Padding(4.f)
 	[
-		SNew(SScrollBox)
-		+ SScrollBox::Slot()
-		[
-			SAssignNew(ContentBox, SVerticalBox)
-			+SVerticalBox::Slot()
-			.AutoHeight()
-			.VAlign(VAlign_Center)
-		]
+		SAssignNew(ContentBox, SScrollBox)
 	];
 
 	ISkeletalMeshEditor* SkeletalMeshEditor = static_cast<ISkeletalMeshEditor*>(HostingApp.Pin().Get());
@@ -95,14 +89,12 @@ void SClothPaintTab::Construct(const FArguments& InArgs)
 		IPersonaToolkit& Persona = SkeletalMeshEditor->GetPersonaToolkit().Get();
 
 		ContentBox->AddSlot()
-		.AutoHeight()
 		[
 			SAssignNew(SelectorWidget, SClothAssetSelector, Persona.GetMesh())
 				.OnSelectionChanged(this, &SClothPaintTab::OnAssetSelectionChanged)
 		];
 
 		ContentBox->AddSlot()
-		.AutoHeight()
 		[
 			DetailsView->AsShared()
 		];
@@ -150,7 +142,6 @@ void SClothPaintTab::UpdatePaintTools()
 			PaintMode->SetPersonaToolKit(SkeletalMeshEditor->GetPersonaToolkit());
 
 			ContentBox->AddSlot()
-			.AutoHeight()
 			[
 				ModeWidget->AsShared()
 			];
diff --git a/Engine/Source/Editor/ClothPainter/Private/SClothPaintTab.h b/Engine/Source/Editor/ClothPainter/Private/SClothPaintTab.h
--- a/Engine/Source/Editor/ClothPainter/Private/SClothPaintTab.h
+++ b/Engine/Source/Editor/ClothPainter/Private/SClothPaintTab.h
@@ -10,6 +10,7 @@ class SClothAssetSelector;
 class IDetailsView;
 class ISkeletalMeshEditor;
 class IPersonaToolkit;
+class SScrollBox;
 
 class CLOTHPAINTER_API SClothPaintTab : public SCompoundWidget
 {
@@ -51,7 +52,7 @@ protected:
 	
 	TSharedPtr<SClothAssetSelector> SelectorWidget;
 	TSharedPtr<SClothPaintWidget> ModeWidget;
-	TSharedPtr<SVerticalBox> ContentBox;
+	TSharedPtr<SScrollBox> ContentBox;
 	TSharedPtr<IDetailsView> DetailsView;
 
 	bool bModeApplied;
