Skeletal Mesh Cloth Panel Reskin Tweaks

diff --git a/Engine/Source/Editor/ClothPainter/Private/SClothAssetSelector.cpp b/Engine/Source/Editor/ClothPainter/Private/SClothAssetSelector.cpp
--- a/Engine/Source/Editor/ClothPainter/Private/SClothAssetSelector.cpp
+++ b/Engine/Source/Editor/ClothPainter/Private/SClothAssetSelector.cpp
@@ -395,12 +395,15 @@ public:
 	{
 		if(InColumnName == Column_Enabled)
 		{
-			return SNew(SCheckBox)
+			return SNew(SBox)
+			.Padding(2.f)
+			[
+				SNew(SCheckBox)
 				.IsEnabled(this, &SMaskListRow::IsMaskCheckboxEnabled, Item)
 				.IsChecked(this, &SMaskListRow::IsMaskEnabledChecked, Item)
 				.OnCheckStateChanged(this, &SMaskListRow::OnMaskEnabledCheckboxChanged, Item)
-				.Padding(2.0f)
-				.ToolTipText(LOCTEXT("MaskEnableCheckBox_ToolTip", "Sets whether this mask is enabled and can affect final parameters for its target parameter."));
+				.ToolTipText(LOCTEXT("MaskEnableCheckBox_ToolTip", "Sets whether this mask is enabled and can affect final parameters for its target parameter."))
+			];
 		}
 
 		if(InColumnName == Column_MaskName)
@@ -798,8 +801,9 @@ void SClothAssetSelector::Construct(const FArguments& InArgs, USkeletalMesh* InM
 				.Padding(0.0f, 0.0f, 4.0f, 0.0f)
 				[
 					SNew(SEditorHeaderButton)
+					.Text(LOCTEXT("CopyClothingFromMeshText_TEXT", "Add Clothing"))
 					.Icon(FAppStyle::Get().GetBrush("Icons.Plus"))
-					.ToolTipText(LOCTEXT("CopyClothingFromMeshText_TEXT", "Copy Clothing from SkeletalMesh"))
+					.ToolTipText(LOCTEXT("CopyClothingFromMeshText_TOOLTIP", "Copy Clothing from SkeletalMesh"))
 					.OnGetMenuContent(this, &SClothAssetSelector::OnGenerateSkeletalMeshPickerForClothCopy)
 				]
 
@@ -871,7 +875,7 @@ void SClothAssetSelector::Construct(const FArguments& InArgs, USkeletalMesh* InM
 					[
 						SNew(SImage)
 						.ColorAndOpacity(FSlateColor::UseForeground())
-						.Image(FAppStyle::Get().GetBrush("Icons.Plus"))
+						.Image(FAppStyle::Get().GetBrush("Icons.PlusCircle"))
 					]
 				]
 			]
@@ -892,8 +896,8 @@ void SClothAssetSelector::Construct(const FArguments& InArgs, USkeletalMesh* InM
 						SNew(SHeaderRow)
 
 						+ SHeaderRow::Column(SMaskListRow::Column_Enabled)
-						.FixedWidth(24)
-						.DefaultLabel(FText::GetEmpty())
+						.FixedWidth(40)
+						.HAlignCell(HAlign_Right) .DefaultLabel(FText::GetEmpty())
 
 						+ SHeaderRow::Column(SMaskListRow::Column_MaskName)
 						.FillWidth(0.5f)
@@ -945,7 +949,7 @@ void SClothAssetSelector::Construct(const FArguments& InArgs, USkeletalMesh* InM
 					.Padding(16.0f, 2.0f, 0.0f, 2.0f)
 					[
 						SNew(STextBlock)
-						.Font(IDetailLayoutBuilder::GetDetailFontBold())
+						.Font(FAppStyle::Get().GetFontStyle("PropertyWindow.NormalFont"))
 						.Text(LOCTEXT("MultipleInfluences", "Use Multiple Influences"))
 					]
 
@@ -954,7 +958,7 @@ void SClothAssetSelector::Construct(const FArguments& InArgs, USkeletalMesh* InM
 					.Padding(16.0f, 2.0f, 0.0f, 2.0f)
 					[
 						SNew(STextBlock)
-						.Font(IDetailLayoutBuilder::GetDetailFontBold())
+						.Font(FAppStyle::Get().GetFontStyle("PropertyWindow.NormalFont"))
 						.Text(LOCTEXT("CurrentRadius", "Kernel Radius"))
 					]
 				]
