Cloth - Add asset exporter interface, hooked from the right click menu of the cloth window.

diff --git a/Engine/Source/Editor/ClothPainter/Private/ClothingAssetExporter.cpp b/Engine/Source/Editor/ClothPainter/Private/ClothingAssetExporter.cpp
new file mode 100644
--- /dev/null
+++ b/Engine/Source/Editor/ClothPainter/Private/ClothingAssetExporter.cpp
@@ -0,0 +1,123 @@
+// Copyright Epic Games, Inc. All Rights Reserved.
+
+#include "ClothingAssetExporter.h"
+#include "ClothingAsset.h"
+#include "Misc/FileHelper.h"
+#include "FileHelpers.h"
+#include "AssetRegistryModule.h"
+#include "IContentBrowserSingleton.h"
+#include "ContentBrowserModule.h"
+
+#define LOCTEXT_NAMESPACE "ClothingAssetExporter"
+
+void ForEachClothingAssetExporter(TFunctionRef<void(UClass*)> Function)
+{
+	// Add clothing asset exporters
+	const TArray<IClothingAssetExporterClassProvider*> ClassProviders = IModularFeatures::Get().GetModularFeatureImplementations<IClothingAssetExporterClassProvider>(IClothingAssetExporterClassProvider::FeatureName);
+	for (const IClothingAssetExporterClassProvider* const ClassProvider : ClassProviders)
+	{
+		if (!ClassProvider)
+		{
+			continue;
+		}
+		if (const TSubclassOf<UClothingAssetExporter> ClothingAssetExporterClass = ClassProvider->GetClothingAssetExporterClass())
+		{
+			UClothingAssetExporter* const ClothingAssetExporter = ClothingAssetExporterClass->GetDefaultObject<UClothingAssetExporter>();
+			if (ClothingAssetExporter)
+			{
+				UClass* const ExportedType = ClothingAssetExporter->GetExportedType();
+
+				Function(ExportedType);
+			}
+		}
+	}
+}
+
+void ExportClothingAsset(const UClothingAssetBase* ClothingAsset, UClass* ExportedType)
+{
+	if (!ClothingAsset || !ExportedType)
+	{
+		return;
+	}
+
+	// Get destination for asset
+	FSaveAssetDialogConfig ExportClothingAssetDialogConfig;
+	{
+		const FString PackageName = ClothingAsset->GetOutermost()->GetName();
+		ExportClothingAssetDialogConfig.DefaultPath = FPackageName::GetLongPackagePath(PackageName);
+		ExportClothingAssetDialogConfig.DefaultAssetName = ClothingAsset->GetName();
+		ExportClothingAssetDialogConfig.AssetClassNames.Add(ExportedType->GetFName());
+		ExportClothingAssetDialogConfig.ExistingAssetPolicy = ESaveAssetDialogExistingAssetPolicy::AllowButWarn;
+		ExportClothingAssetDialogConfig.DialogTitleOverride = LOCTEXT("ExportClothingAssetDialogTitle", "Export Clothing Asset As");
+	}
+
+	FContentBrowserModule& ContentBrowserModule = FModuleManager::LoadModuleChecked<FContentBrowserModule>("ContentBrowser");
+
+	FString NewPackageName;
+	FText OutError;
+	for (bool bFilenameValid = false; !bFilenameValid; bFilenameValid = FFileHelper::IsFilenameValidForSaving(NewPackageName, OutError))
+	{
+		const FString ExportClothingAssetPath = ContentBrowserModule.Get().CreateModalSaveAssetDialog(ExportClothingAssetDialogConfig);
+		if (ExportClothingAssetPath.IsEmpty())
+		{
+			return;
+		}
+		NewPackageName = FPackageName::ObjectPathToPackageName(ExportClothingAssetPath);
+	}
+
+	// Try loading the package, when overwriting an existing asset it needs to be loaded first
+	LoadPackage(nullptr, *NewPackageName, LOAD_Quiet | LOAD_EditorOnly);
+
+	// Find a matching exporter
+	UClothingAssetExporter* ClothingAssetExporter = nullptr;
+
+	const TArray<IClothingAssetExporterClassProvider*> ClassProviders = IModularFeatures::Get().GetModularFeatureImplementations<IClothingAssetExporterClassProvider>(IClothingAssetExporterClassProvider::FeatureName);
+	for (IClothingAssetExporterClassProvider* const ClassProvider : ClassProviders)
+	{
+		if (!ClassProvider)
+		{
+			continue;
+		}
+		if (const TSubclassOf<UClothingAssetExporter> ClothingAssetExporterClass = ClassProvider->GetClothingAssetExporterClass())
+		{
+			UClothingAssetExporter* const ClothingAssetExporterCandidate = ClothingAssetExporterClass->GetDefaultObject<UClothingAssetExporter>();
+
+			if (ClothingAssetExporterCandidate && ExportedType == ClothingAssetExporterCandidate->GetExportedType())
+			{
+				ClothingAssetExporter = ClothingAssetExporterCandidate;
+				break;
+			}
+		}
+	}
+
+	if (!ClothingAssetExporter)
+	{
+		const FText TitleMessage = LOCTEXT("ClothingAssetExporterTitleMessage", "Error Exporting Clothing Asset");
+		const FText ErrorMessage = FText::Format(LOCTEXT("ClothingAssetExporterErrorMessage", "The Asset's exporter for type {0} cannot be found anymore. Check that the plugin hasn't been unloaded."), FText::FromName(ExportedType->GetFName()));
+		FMessageDialog::Open(EAppMsgType::Ok, EAppReturnType::Ok, ErrorMessage, &TitleMessage);
+		return;
+	}
+
+	// Create the new asset
+	const FName NewAssetName(FPackageName::GetLongPackageAssetName(NewPackageName));
+	UPackage* const NewPackage = CreatePackage(*NewPackageName);
+
+	UObject* const NewAsset = NewObject<UObject>(NewPackage, ExportedType, NewAssetName, RF_Public | RF_Standalone | RF_Transactional);
+
+	NewAsset->MarkPackageDirty();
+
+	// Notify the asset registry
+	FAssetRegistryModule::AssetCreated(NewAsset);
+
+	// Call the export function
+	ClothingAssetExporter->Export(ClothingAsset, NewAsset);
+
+	// Save the package
+	TArray<UPackage*> PackagesToSave;
+	PackagesToSave.Add(NewAsset->GetOutermost());
+	constexpr bool bCheckDirty = false;
+	constexpr bool bPromptToSave = false;
+	FEditorFileUtils::PromptForCheckoutAndSave(PackagesToSave, bCheckDirty, bPromptToSave);
+}
+
+#undef LOCTEXT_NAMESPACE
diff --git a/Engine/Source/Editor/ClothPainter/Private/ClothingAssetListCommands.cpp b/Engine/Source/Editor/ClothPainter/Private/ClothingAssetListCommands.cpp
--- a/Engine/Source/Editor/ClothPainter/Private/ClothingAssetListCommands.cpp
+++ b/Engine/Source/Editor/ClothPainter/Private/ClothingAssetListCommands.cpp
@@ -1,6 +1,7 @@
 // Copyright Epic Games, Inc. All Rights Reserved.
 
 #include "ClothingAssetListCommands.h"
+#include "ClothingAssetExporter.h"
 
 #define LOCTEXT_NAMESPACE "ClothingAssetListCommands"
 
@@ -11,6 +12,41 @@ void FClothingAssetListCommands::RegisterCommands()
 	UI_COMMAND(ReimportAsset, "Reimport Asset", "If an asset was originally imported from an external file, reimport from that file.", EUserInterfaceActionType::Button, FInputChord());
 #endif
 	UI_COMMAND(RebuildAssetParams, "Rebuild Asset Parameter Masks", "Takes the parameter masks in LOD0 and creates masks in all lower LODs to match, casting those parameter masks to the new mesh.", EUserInterfaceActionType::Button, FInputChord());
+
+	// Add one command per exporter provider
+	const TArray<IClothingAssetExporterClassProvider*> ClassProviders = IModularFeatures::Get().GetModularFeatureImplementations<IClothingAssetExporterClassProvider>(IClothingAssetExporterClassProvider::FeatureName);
+	const int32 NumClassProviders = ClassProviders.Num();
+	ExportAssets.Reserve(NumClassProviders);
+
+	for (const IClothingAssetExporterClassProvider* const ClassProvider : ClassProviders)
+	{
+		if (ClassProvider)
+		{
+			if (const TSubclassOf<UClothingAssetExporter> ClothingAssetExporterClass = ClassProvider->GetClothingAssetExporterClass())
+			{
+				UClothingAssetExporter* const ClothingAssetExporter = ClothingAssetExporterClass->GetDefaultObject<UClothingAssetExporter>();
+				if (ClothingAssetExporter)
+				{
+					UClass* ExportedType = ClothingAssetExporter->GetExportedType();
+					const FName CommandName = ExportedType->GetFName();
+					const FText ShortName = FText::Format(LOCTEXT("ExportToClothAssetType", "Export To {0}"), FText::FromName(CommandName));
+					const FText Description = FText::Format(LOCTEXT("ExportToClothAssetTypeDescription", "Export the selected Clothing Data to the {0} asset type."), FText::FromName(CommandName));
+
+					TSharedPtr<FUICommandInfo>& CommandId = ExportAssets.Add(CommandName);
+
+					FUICommandInfo::MakeCommandInfo(
+						this->AsShared(),
+						CommandId,
+						CommandName,
+						ShortName,
+						Description,
+						FSlateIcon(),
+						EUserInterfaceActionType::Button,
+						FInputChord());
+				}
+			}
+		}
+	}
 }
 
 #undef LOCTEXT_NAMESPACE
diff --git a/Engine/Source/Editor/ClothPainter/Private/ClothingAssetListCommands.h b/Engine/Source/Editor/ClothPainter/Private/ClothingAssetListCommands.h
--- a/Engine/Source/Editor/ClothPainter/Private/ClothingAssetListCommands.h
+++ b/Engine/Source/Editor/ClothPainter/Private/ClothingAssetListCommands.h
@@ -26,5 +26,5 @@ public:
 #endif
 
 	TSharedPtr<FUICommandInfo> RebuildAssetParams;
-
+	TMap<FName, TSharedPtr<FUICommandInfo>> ExportAssets;
 };
diff --git a/Engine/Source/Editor/ClothPainter/Private/SClothAssetSelector.cpp b/Engine/Source/Editor/ClothPainter/Private/SClothAssetSelector.cpp
--- a/Engine/Source/Editor/ClothPainter/Private/SClothAssetSelector.cpp
+++ b/Engine/Source/Editor/ClothPainter/Private/SClothAssetSelector.cpp
@@ -35,6 +35,7 @@
 #include "Editor/ContentBrowser/Private/SAssetPicker.h"
 #include "ContentBrowserModule.h"
 #include "Animation/DebugSkelMeshComponent.h"
+#include "ClothingAssetExporter.h"
 
 #define LOCTEXT_NAMESPACE "ClothAssetSelector"
 
@@ -160,6 +161,17 @@ public:
 			FExecuteAction::CreateSP(this, &SAssetListRow::RebuildLODParameters),
 			FCanExecuteAction::CreateSP(this, &SAssetListRow::CanRebuildLODParameters)
 		);
+
+		// Add clothing asset exporters
+		ForEachClothingAssetExporter([this, &Commands](UClass* ExportedType)
+			{
+				if (const TSharedPtr<FUICommandInfo>* const CommandId = Commands.ExportAssets.Find(ExportedType->GetFName()))
+				{
+					UICommandList->MapAction(
+						*CommandId,
+						FExecuteAction::CreateSP(this, &SAssetListRow::ExportAsset, ExportedType));
+				}
+			});
 	}
 
 	virtual FReply OnMouseButtonUp(const FGeometry& MyGeometry, const FPointerEvent& MouseEvent) override
@@ -176,6 +188,11 @@ public:
 				Builder.AddMenuEntry(Commands.ReimportAsset);
 #endif
 				Builder.AddMenuEntry(Commands.RebuildAssetParams);
+
+				for (const TPair<FName, TSharedPtr<FUICommandInfo>>& CommandId : Commands.ExportAssets)
+				{
+					Builder.AddMenuEntry(CommandId.Value);
+				}
 			}
 			Builder.EndSection();
 
@@ -342,6 +359,19 @@ private:
 		}
 	}
 
+	void ExportAsset(UClass* ExportedType)
+	{
+		if (!Item.IsValid() || !ExportedType)
+		{
+			return;
+		}
+
+		if (const UClothingAssetCommon* const ClothingAsset = Item->ClothingAsset.Get())
+		{
+			ExportClothingAsset(ClothingAsset, ExportedType);
+		}
+	}
+
 	bool CanRebuildLODParameters() const
 	{
 		if(!Item.IsValid())
diff --git a/Engine/Source/Editor/ClothPainter/Public/ClothingAssetExporter.h b/Engine/Source/Editor/ClothPainter/Public/ClothingAssetExporter.h
new file mode 100644
--- /dev/null
+++ b/Engine/Source/Editor/ClothPainter/Public/ClothingAssetExporter.h
@@ -0,0 +1,54 @@
+// Copyright Epic Games, Inc. All Rights Reserved.
+
+#pragma once
+
+#include "UObject/Object.h"
+#include "Features/IModularFeature.h"
+#include "Templates/SubclassOf.h"
+
+#include "ClothingAssetExporter.generated.h"
+
+class UClothingAssetBase;
+class UClothingAssetExporter;
+
+/** A modular interface to provide ways to export clothing assets. */
+class CLOTHPAINTER_API IClothingAssetExporterClassProvider : public IModularFeature
+{
+public:
+	inline static const FName FeatureName = TEXT("ClothingAssetExporterClassProvider");
+
+	virtual TSubclassOf<UClothingAssetExporter> GetClothingAssetExporterClass() const = 0;
+};
+
+/** Modular exporter base class. */
+UCLASS(Abstract)
+class CLOTHPAINTER_API UClothingAssetExporter : public UObject
+{
+	GENERATED_BODY()
+
+public:
+	/** Return the class of the exported asset type. */
+	virtual UClass* GetExportedType() const
+	PURE_VIRTUAL(UClothingAssetExporter::GetExportedType, return nullptr;);
+
+	/**
+	 * Export the specified asset.
+	 * \param ClothingAsset the source clothing asset to be exported.
+	 * \param ExportedAsset the destination asset object in the exported type provided by the caller, ready to be filled by the Export function.
+	 */
+	virtual void Export(const UClothingAssetBase* ClothingAsset, UObject* ExportedAsset)
+	PURE_VIRTUAL(UClothingAssetExporter::Export,);
+};
+
+/**
+ * Call the specified function for each asset exporter plugin currently loaded, passing the target exported type as a parameter.
+ * \param Function the function to call for each asset exporter found.
+ */
+void CLOTHPAINTER_API ForEachClothingAssetExporter(TFunctionRef<void(UClass*)> Function);
+
+/**
+ * Open a dialog and create a new asset of the specified type, then call the associated export function on it.
+ * \param ClothingAsset the asset to export.
+ * \param ExportedType the exported type.
+ */
+void CLOTHPAINTER_API ExportClothingAsset(const UClothingAssetBase* ClothingAsset, UClass* ExportedType);
