Cloth - Fixed crash when undoing paint on an unbound clothing data.

diff --git a/Engine/Source/Editor/ClothPainter/Private/ClothMeshAdapter.cpp b/Engine/Source/Editor/ClothPainter/Private/ClothMeshAdapter.cpp
--- a/Engine/Source/Editor/ClothPainter/Private/ClothMeshAdapter.cpp
+++ b/Engine/Source/Editor/ClothPainter/Private/ClothMeshAdapter.cpp
@@ -269,57 +269,60 @@ bool FClothMeshPaintAdapter::InitializeVertexData()
 		{
 			if(DebugComponent->SkinnedSelectedClothingPositions.Num() > 0)
 			{
-				UClothingAssetCommon* ConcreteAsset = CastChecked<UClothingAssetCommon>(SelectedAsset);
-				const FClothLODDataCommon& LODData = ConcreteAsset->LodData[PaintingClothLODIndex];
-				const FClothPhysicalMeshData& MeshData = LODData.PhysicalMeshData;
+				UClothingAssetCommon* const ConcreteAsset = CastChecked<UClothingAssetCommon>(SelectedAsset);
+				if (ConcreteAsset->LodData.IsValidIndex(PaintingClothLODIndex))
+				{
+					const FClothLODDataCommon& LODData = ConcreteAsset->LodData[PaintingClothLODIndex];
+					const FClothPhysicalMeshData& MeshData = LODData.PhysicalMeshData;
 
-				MeshVertices.Append(DebugComponent->SkinnedSelectedClothingPositions);
-				MeshIndices.Append(MeshData.Indices);
+					MeshVertices.Append(DebugComponent->SkinnedSelectedClothingPositions);
+					MeshIndices.Append(MeshData.Indices);
 
-				for(int32 Index = IndexOffset; Index < MeshIndices.Num(); ++Index)
-				{
-					MeshIndices[Index] += VertexOffset;
-				}
+					for(int32 Index = IndexOffset; Index < MeshIndices.Num(); ++Index)
+					{
+						MeshIndices[Index] += VertexOffset;
+					}
 
-				FClothAssetInfo Info;
-				Info.IndexStart = IndexOffset;
-				Info.VertexStart = VertexOffset;
+					FClothAssetInfo Info;
+					Info.IndexStart = IndexOffset;
+					Info.VertexStart = VertexOffset;
 
-				IndexOffset += MeshData.Indices.Num();
-				VertexOffset += MeshData.Vertices.Num();
+					IndexOffset += MeshData.Indices.Num();
+					VertexOffset += MeshData.Vertices.Num();
 
-				Info.IndexEnd = IndexOffset;
-				Info.VertexEnd = VertexOffset;
-				Info.Asset = ConcreteAsset;
+					Info.IndexEnd = IndexOffset;
+					Info.VertexEnd = VertexOffset;
+					Info.Asset = ConcreteAsset;
 
-				// Set up the edge map / neighbor map
+					// Set up the edge map / neighbor map
 
-				// Pre fill the map, 1 per index
-				Info.NeighborMap.AddDefaulted(Info.VertexEnd - Info.VertexStart);
+					// Pre fill the map, 1 per index
+					Info.NeighborMap.AddDefaulted(Info.VertexEnd - Info.VertexStart);
 
-				// Fill in neighbors defined by triangles
-				const int32 NumTris = MeshIndices.Num() / 3;
-				for(int32 TriIndex = 0; TriIndex < NumTris; ++TriIndex)
-				{
-					const int32 I0 = MeshIndices[TriIndex * 3];
-					const int32 I1 = MeshIndices[TriIndex * 3 + 1];
-					const int32 I2 = MeshIndices[TriIndex * 3 + 2];
+					// Fill in neighbors defined by triangles
+					const int32 NumTris = MeshIndices.Num() / 3;
+					for(int32 TriIndex = 0; TriIndex < NumTris; ++TriIndex)
+					{
+						const int32 I0 = MeshIndices[TriIndex * 3];
+						const int32 I1 = MeshIndices[TriIndex * 3 + 1];
+						const int32 I2 = MeshIndices[TriIndex * 3 + 2];
 
-					TArray<int32>& I0Neighbors = Info.NeighborMap[I0];
-					TArray<int32>& I1Neighbors = Info.NeighborMap[I1];
-					TArray<int32>& I2Neighbors = Info.NeighborMap[I2];
+						TArray<int32>& I0Neighbors = Info.NeighborMap[I0];
+						TArray<int32>& I1Neighbors = Info.NeighborMap[I1];
+						TArray<int32>& I2Neighbors = Info.NeighborMap[I2];
 
-					I0Neighbors.AddUnique(I1);
-					I0Neighbors.AddUnique(I2);
+						I0Neighbors.AddUnique(I1);
+						I0Neighbors.AddUnique(I2);
 
-					I1Neighbors.AddUnique(I0);
-					I1Neighbors.AddUnique(I2);
+						I1Neighbors.AddUnique(I0);
+						I1Neighbors.AddUnique(I2);
 
-					I2Neighbors.AddUnique(I0);
-					I2Neighbors.AddUnique(I1);
-				}
+						I2Neighbors.AddUnique(I0);
+						I2Neighbors.AddUnique(I1);
+					}
 
-				AssetInfoMap.Add(Info);
+					AssetInfoMap.Add(Info);
+				}
 			}
 		}
 	}
