Fixes clash between Skeletal Mesh Editing toolset and Cloth Paint toolset.

diff --git a/Engine/Source/Editor/ClothPainter/Private/ClothPaintingModule.cpp b/Engine/Source/Editor/ClothPainter/Private/ClothPaintingModule.cpp
--- a/Engine/Source/Editor/ClothPainter/Private/ClothPaintingModule.cpp
+++ b/Engine/Source/Editor/ClothPainter/Private/ClothPaintingModule.cpp
@@ -5,6 +5,7 @@
 #include "ClothPaintToolCommands.h"
 #include "ClothPainterCommands.h"
 #include "ClothingPaintEditMode.h"
+#include "EditorModeManager.h"
 #include "Delegates/Delegate.h"
 #include "EditorModeRegistry.h"
 #include "Framework/Commands/UIAction.h"
@@ -130,7 +131,7 @@ TSharedRef<FExtender> FClothPaintingModule::ExtendSkelMeshEditorToolbar(const TS
 	InCommandList->MapAction(FClothPainterCommands::Get().TogglePaintMode,
 		FExecuteAction::CreateRaw(this, &FClothPaintingModule::OnToggleMode, Ptr),
 		FCanExecuteAction(),
-		FIsActionChecked::CreateRaw(this, &FClothPaintingModule::GetIsPaintToolsButtonChecked, Ptr)
+		FIsActionChecked::CreateRaw(this, &FClothPaintingModule::IsPaintModeActive, Ptr)
 	);
 
 	return ToolbarExtender.ToSharedRef();
@@ -161,14 +162,9 @@ void FClothPaintingModule::RegisterMenus()
 
 FText FClothPaintingModule::GetPaintToolsButtonText(TWeakPtr<ISkeletalMeshEditor> InSkeletalMeshEditor) const
 {
-	TSharedPtr<SClothPaintTab> ClothTab = GetActiveClothTab(InSkeletalMeshEditor, false);
-
-	if(ClothTab.IsValid())
+	if(IsPaintModeActive(InSkeletalMeshEditor))
 	{
-		if(ClothTab->IsPaintModeActive())
-		{
-			return LOCTEXT("ToggleButton_Deactivate", "Deactivate Cloth Paint");
-		}
+		return LOCTEXT("ToggleButton_Deactivate", "Deactivate Cloth Paint");
 	}
 
 	return LOCTEXT("ToggleButton_Activate", "Activate Cloth Paint");
@@ -176,38 +172,44 @@ FText FClothPaintingModule::GetPaintToolsButtonText(TWeakPtr<ISkeletalMeshEditor
 
 FText FClothPaintingModule::GetPaintToolsButtonToolTip(TWeakPtr<ISkeletalMeshEditor> InSkeletalMeshEditor) const
 {
-	TSharedPtr<SClothPaintTab> ClothTab = GetActiveClothTab(InSkeletalMeshEditor, false);
-
-	if(ClothTab.IsValid())
+	if(IsPaintModeActive(InSkeletalMeshEditor))
 	{
-		if(ClothTab->IsPaintModeActive())
-		{
-			return LOCTEXT("ToggleButton_Deactivate_ToolTip", "Deactivate the cloth paint tool, and go back to the current selection mode.");
-		}
+		return LOCTEXT("ToggleButton_Deactivate_ToolTip", "Deactivate the cloth paint tool, and go back to the current selection mode.");
 	}
 
 	return LOCTEXT("ToggleButton_Activate_ToolTip", "Activate the cloth paint tool, and open the Clothing window to allow selection of the clothing assets and of their paint targets.");
 }
 
-bool FClothPaintingModule::GetIsPaintToolsButtonChecked(TWeakPtr<ISkeletalMeshEditor> InSkeletalMeshEditor) const
+bool FClothPaintingModule::IsPaintModeActive(TWeakPtr<ISkeletalMeshEditor> InSkeletalMeshEditor) const
 {
-	TSharedPtr<SClothPaintTab> ClothTab = GetActiveClothTab(InSkeletalMeshEditor, false);
-
-	if(ClothTab.IsValid())
+	if (!GetActiveClothTab(InSkeletalMeshEditor, false /* don't invoke tab*/).IsValid())
 	{
-		return ClothTab->IsPaintModeActive();
+		// haven't initialized UI yet
+		return false;
 	}
-
-	return false;
+		
+	TSharedPtr<ISkeletalMeshEditor> SkeletalMeshEditor = InSkeletalMeshEditor.Pin();
+	return SkeletalMeshEditor.IsValid() && SkeletalMeshEditor->GetEditorModeManager().IsModeActive(PaintModeID);
 }
 
-void FClothPaintingModule::OnToggleMode(TWeakPtr<ISkeletalMeshEditor> InSkeletalMeshEditor)
+void FClothPaintingModule::OnToggleMode(TWeakPtr<ISkeletalMeshEditor> InSkeletalMeshEditor) const
 {
-	TSharedPtr<SClothPaintTab> ClothTab = GetActiveClothTab(InSkeletalMeshEditor);
+	ISkeletalMeshEditor* SkeletalMeshEditor = InSkeletalMeshEditor.Pin().Get();
+	FEditorModeTools& ModeManager = SkeletalMeshEditor->GetEditorModeManager();
 
-	if(ClothTab.IsValid())
+	if (!IsPaintModeActive(InSkeletalMeshEditor))
+	{
+		ModeManager.ActivateMode(PaintModeID, true);
+		FClothingPaintEditMode* PaintMode = (FClothingPaintEditMode*)SkeletalMeshEditor->GetEditorModeManager().GetActiveMode(PaintModeID);
+		if (PaintMode)
+		{
+			PaintMode->SetPersonaToolKit(SkeletalMeshEditor->GetPersonaToolkit());
+			PaintMode->SetupClothPaintTab(GetActiveClothTab(InSkeletalMeshEditor));
+		}
+	}
+	else
 	{
-		ClothTab->TogglePaintMode();
+		ModeManager.DeactivateMode(PaintModeID);
 	}
 }
 
diff --git a/Engine/Source/Editor/ClothPainter/Private/ClothingPaintEditMode.cpp b/Engine/Source/Editor/ClothPainter/Private/ClothingPaintEditMode.cpp
--- a/Engine/Source/Editor/ClothPainter/Private/ClothingPaintEditMode.cpp
+++ b/Engine/Source/Editor/ClothPainter/Private/ClothingPaintEditMode.cpp
@@ -17,6 +17,7 @@
 #include "Engine/SkeletalMesh.h"
 #include "IPersonaPreviewScene.h"
 #include "IPersonaToolkit.h"
+#include "SClothPaintTab.h"
 #include "Misc/CoreMiscDefines.h"
 #include "Misc/Guid.h"
 #include "ShowFlags.h"
@@ -62,6 +63,12 @@ void FClothingPaintEditMode::SetPersonaToolKit(class TSharedPtr<IPersonaToolkit>
 	PersonaToolkit = InToolkit;
 }
 
+void FClothingPaintEditMode::SetupClothPaintTab(TSharedPtr<SClothPaintTab> InClothPaintTab)
+{
+	ClothPaintTab = InClothPaintTab;
+	ClothPaintTab.Pin().Get()->EnterPaintMode();
+}
+
 void FClothingPaintEditMode::Enter()
 {
 	IMeshPaintEdMode::Enter();
@@ -150,5 +157,10 @@ void FClothingPaintEditMode::Exit()
 
 	ClothPainter->ExitPaintMode();
 
+	if (ClothPaintTab.IsValid())
+	{
+		ClothPaintTab.Pin().Get()->ExitPaintMode();	
+	}
+
 	IMeshPaintEdMode::Exit();
 }
diff --git a/Engine/Source/Editor/ClothPainter/Private/ClothingPaintEditMode.h b/Engine/Source/Editor/ClothPainter/Private/ClothingPaintEditMode.h
--- a/Engine/Source/Editor/ClothPainter/Private/ClothingPaintEditMode.h
+++ b/Engine/Source/Editor/ClothPainter/Private/ClothingPaintEditMode.h
@@ -5,6 +5,7 @@
 #include "IMeshPaintMode.h"
 #include "Templates/SharedPointer.h"
 
+class SClothPaintTab;
 class FClothPainter;
 class IPersonaToolkit;
 
@@ -19,6 +20,7 @@ public:
 	virtual TSharedPtr<class FModeToolkit> GetToolkit() override;
 
 	void SetPersonaToolKit(class TSharedPtr<IPersonaToolkit> InToolkit);
+	void SetupClothPaintTab(TSharedPtr<SClothPaintTab> InClothPaintTab);
 
 	virtual void Enter() override;
 	virtual void Exit() override;
@@ -29,4 +31,5 @@ protected:
 
 	class IPersonaPreviewScene* GetAnimPreviewScene() const;
 	class TWeakPtr<IPersonaToolkit> PersonaToolkit;
+	TWeakPtr<SClothPaintTab> ClothPaintTab;
 };
\ No newline at end of file
diff --git a/Engine/Source/Editor/ClothPainter/Private/SClothPaintTab.cpp b/Engine/Source/Editor/ClothPainter/Private/SClothPaintTab.cpp
--- a/Engine/Source/Editor/ClothPainter/Private/SClothPaintTab.cpp
+++ b/Engine/Source/Editor/ClothPainter/Private/SClothPaintTab.cpp
@@ -34,16 +34,15 @@ struct FGeometry;
 
 #define LOCTEXT_NAMESPACE "SClothPaintTab"
 
-SClothPaintTab::SClothPaintTab() 
-	: bModeApplied(false), bPaintModeEnabled(false)
+SClothPaintTab::SClothPaintTab() : bModeApplied(false)
 {	
 }
 
 SClothPaintTab::~SClothPaintTab()
 {
-	if(ISkeletalMeshEditor* SkeletalMeshEditor = static_cast<ISkeletalMeshEditor*>(HostingApp.Pin().Get()))
+	if(const ISkeletalMeshEditor* Editor = static_cast<ISkeletalMeshEditor*>(HostingApp.Pin().Get()))
 	{
-		SkeletalMeshEditor->GetEditorModeManager().ActivateDefaultMode();
+		Editor->GetEditorModeManager().ActivateDefaultMode();
 	}
 }
 
@@ -83,9 +82,8 @@ void SClothPaintTab::Construct(const FArguments& InArgs)
 		SAssignNew(ContentBox, SScrollBox)
 	];
 
-	ISkeletalMeshEditor* SkeletalMeshEditor = static_cast<ISkeletalMeshEditor*>(HostingApp.Pin().Get());
-
-	if(SkeletalMeshEditor)
+	
+	if(ISkeletalMeshEditor* SkeletalMeshEditor = GetSkeletalMeshEditor())
 	{
 		IPersonaToolkit& Persona = SkeletalMeshEditor->GetPersonaToolkit().Get();
 
@@ -107,81 +105,61 @@ void SClothPaintTab::Tick(const FGeometry& AllottedGeometry, const double InCurr
 	SCompoundWidget::Tick(AllottedGeometry, InCurrentTime, InDeltaTime);
 }
 
-void SClothPaintTab::TogglePaintMode()
+void SClothPaintTab::EnterPaintMode()
 {
-	bPaintModeEnabled = !bPaintModeEnabled;
-	UpdatePaintTools();
-}
-
-bool SClothPaintTab::IsPaintModeActive() const
-{
-	return bPaintModeEnabled;
-}
-
-void SClothPaintTab::UpdatePaintTools()
-{
-	if(!HostingApp.IsValid())
+	const ISkeletalMeshEditor* SkeletalMeshEditor = GetSkeletalMeshEditor();
+	if (!SkeletalMeshEditor)
+	{
+		return;
+	}
+	FClothingPaintEditMode* PaintMode = (FClothingPaintEditMode*)SkeletalMeshEditor->GetEditorModeManager().GetActiveMode(PaintModeID);
+	if (!PaintMode)
 	{
-		// If we have no valid host, we can't do anything with our editor below, so don't perform a
-		// state update on the tool mode.
 		return;
 	}
+	
+	FClothPainter* ClothPainter = static_cast<FClothPainter*>(PaintMode->GetMeshPainter());
+	check(ClothPainter);
+
+	ClothPainter->Reset();
+	ModeWidget = StaticCastSharedPtr<SClothPaintWidget>(ClothPainter->GetWidget());
+	
+	ContentBox->AddSlot()
+	[
+		ModeWidget->AsShared()
+	];
 
-	if (bPaintModeEnabled)
+	if(SelectorWidget.IsValid())
 	{
-		ISkeletalMeshEditor* SkeletalMeshEditor = static_cast<ISkeletalMeshEditor*>(HostingApp.Pin().Get());
-		SkeletalMeshEditor->GetEditorModeManager().ActivateMode(PaintModeID, true);
+		TWeakObjectPtr<UClothingAssetCommon> WeakAsset = SelectorWidget->GetSelectedAsset();
 
-		FClothingPaintEditMode* PaintMode = (FClothingPaintEditMode*)SkeletalMeshEditor->GetEditorModeManager().GetActiveMode(PaintModeID);
-		if (PaintMode)
+		if(WeakAsset.Get())
 		{
-			FClothPainter* ClothPainter = static_cast<FClothPainter*>(PaintMode->GetMeshPainter());
-			check(ClothPainter);
-
-			ClothPainter->Reset();
-			ModeWidget = StaticCastSharedPtr<SClothPaintWidget>(ClothPainter->GetWidget());
-			PaintMode->SetPersonaToolKit(SkeletalMeshEditor->GetPersonaToolkit());
-
-			ContentBox->AddSlot()
-			[
-				ModeWidget->AsShared()
-			];
-
-			if(SelectorWidget.IsValid())
-			{
-				TWeakObjectPtr<UClothingAssetCommon> WeakAsset = SelectorWidget->GetSelectedAsset();
-
-				if(WeakAsset.Get())
-				{
-					ClothPainter->OnAssetSelectionChanged(WeakAsset.Get(), SelectorWidget->GetSelectedLod(), SelectorWidget->GetSelectedMask());
-				}
-			}
+			ClothPainter->OnAssetSelectionChanged(WeakAsset.Get(), SelectorWidget->GetSelectedLod(), SelectorWidget->GetSelectedMask());
 		}
 	}
-	else
-	{
-		ContentBox->RemoveSlot(ModeWidget->AsShared());
-		ISkeletalMeshEditor* SkeletalMeshEditor = static_cast<ISkeletalMeshEditor*>(HostingApp.Pin().Get());
-		SkeletalMeshEditor->GetEditorModeManager().ActivateDefaultMode();
-		ModeWidget = nullptr;
-	}
+}
+
+void SClothPaintTab::ExitPaintMode()
+{
+	ContentBox->RemoveSlot(ModeWidget->AsShared());
+	ModeWidget = nullptr;
 }
 
 void SClothPaintTab::OnAssetSelectionChanged(TWeakObjectPtr<UClothingAssetCommon> InAssetPtr, int32 InLodIndex, int32 InMaskIndex)
 {
-	if(bPaintModeEnabled)
+	const ISkeletalMeshEditor* SkeletalMeshEditor = GetSkeletalMeshEditor();
+	if (!SkeletalMeshEditor)
 	{
-		ISkeletalMeshEditor* SkeletalMeshEditor = static_cast<ISkeletalMeshEditor*>(HostingApp.Pin().Get());
+		return;
+	}
 
-		FClothingPaintEditMode* PaintMode = (FClothingPaintEditMode*)SkeletalMeshEditor->GetEditorModeManager().GetActiveMode(PaintModeID);
-		if(PaintMode)
+	FClothingPaintEditMode* PaintMode = (FClothingPaintEditMode*)SkeletalMeshEditor->GetEditorModeManager().GetActiveMode(PaintModeID);
+	if(PaintMode)
+	{
+		if(FClothPainter* ClothPainter = static_cast<FClothPainter*>(PaintMode->GetMeshPainter()))
 		{
-			FClothPainter* ClothPainter = static_cast<FClothPainter*>(PaintMode->GetMeshPainter());
-
-			if(ClothPainter)
-			{
-				ClothPainter->OnAssetSelectionChanged(InAssetPtr.Get(), InLodIndex, InMaskIndex);
-			}
+			ClothPainter->OnAssetSelectionChanged(InAssetPtr.Get(), InLodIndex, InMaskIndex);
 		}
 	}
 
diff --git a/Engine/Source/Editor/ClothPainter/Private/SClothPaintTab.h b/Engine/Source/Editor/ClothPainter/Private/SClothPaintTab.h
--- a/Engine/Source/Editor/ClothPainter/Private/SClothPaintTab.h
+++ b/Engine/Source/Editor/ClothPainter/Private/SClothPaintTab.h
@@ -32,17 +32,12 @@ public:
 	void Construct(const FArguments& InArgs);
 	virtual void Tick(const FGeometry& AllottedGeometry, const double InCurrentTime, const float InDeltaTime) override;
 
-	/** Toggles cloth paint mode */
-	void TogglePaintMode();
-
-	/** Gets whether cloth paint mode is active */
-	bool IsPaintModeActive() const;
+	/** Setup and teardown the cloth paint UI */
+	void EnterPaintMode();
+	void ExitPaintMode();
 
 protected:
 
-	/** Called as the tool selection changes to enable/disable painting */
-	void UpdatePaintTools();
-
 	/** Called from the selector when the asset selection changes (Asset, LOD, Mask) */
 	void OnAssetSelectionChanged(TWeakObjectPtr<UClothingAssetCommon> InAssetPtr, int32 InLodIndex, int32 InMaskIndex);
 
@@ -61,5 +56,4 @@ protected:
 	TSharedPtr<IDetailsView> DetailsView;
 
 	bool bModeApplied;
-	bool bPaintModeEnabled;
 };
diff --git a/Engine/Source/Editor/ClothPainter/Public/ClothPaintingModule.h b/Engine/Source/Editor/ClothPainter/Public/ClothPaintingModule.h
--- a/Engine/Source/Editor/ClothPainter/Public/ClothPaintingModule.h
+++ b/Engine/Source/Editor/ClothPainter/Public/ClothPaintingModule.h
@@ -57,11 +57,11 @@ private:
 	// Gets tool tip for the enable paint tools button
 	FText GetPaintToolsButtonToolTip(TWeakPtr<ISkeletalMeshEditor> InSkeletalMeshEditor) const;
 
-	// Whether paint mode is active
-	bool GetIsPaintToolsButtonChecked(TWeakPtr<ISkeletalMeshEditor> InSkeletalMeshEditor) const;
-
 	// Toggles paint mode on the clothing tab
-	void OnToggleMode(TWeakPtr<ISkeletalMeshEditor> InSkeletalMeshEditor);
+	void OnToggleMode(TWeakPtr<ISkeletalMeshEditor> InSkeletalMeshEditor) const;
+
+	// Return true if currently in paint mode
+	bool IsPaintModeActive(TWeakPtr<ISkeletalMeshEditor> InSkeletalMeshEditor) const;
 
 	// Gets the current active clothing tab, will invoke (spawn or draw attention to) if bInvoke == true
 	TSharedPtr<SClothPaintTab> GetActiveClothTab(TWeakPtr<ISkeletalMeshEditor> InSkeletalMeshEditor, bool bInvoke = true) const;
