Cloth:
- Fixed some minor UI discrepencies such as fonts, style, and margins.
- Hid UI that were ClothingAssetCommon specific for all other types.

#rnx

diff --git a/Engine/Source/Editor/ClothPainter/Private/ClothMeshAdapter.cpp b/Engine/Source/Editor/ClothPainter/Private/ClothMeshAdapter.cpp
--- a/Engine/Source/Editor/ClothPainter/Private/ClothMeshAdapter.cpp
+++ b/Engine/Source/Editor/ClothPainter/Private/ClothMeshAdapter.cpp
@@ -214,22 +214,24 @@ void FClothMeshPaintAdapter::SetSelectedClothingAsset(const FGuid& InAssetGuid,
 	{
 		for(UClothingAssetBase* Asset : ReferencedSkeletalMesh->GetMeshClothingAssets())
 		{
-			UClothingAssetCommon* ConcreteAsset = CastChecked<UClothingAssetCommon>(Asset);
-			if(ConcreteAsset->GetAssetGuid() == InAssetGuid)
+			if(const UClothingAssetCommon* const ConcreteAsset = Cast<const UClothingAssetCommon>(Asset))
 			{
-				if(ConcreteAsset->IsValidLod(InAssetLod))
+				if(ConcreteAsset->GetAssetGuid() == InAssetGuid)
 				{
-					FClothLODDataCommon& LodData = ConcreteAsset->LodData[InAssetLod];
-
-					if(LodData.PointWeightMaps.IsValidIndex(InMaskIndex))
+					if(ConcreteAsset->IsValidLod(InAssetLod))
 					{
-						PaintingClothLODIndex = InAssetLod;
-						PaintingClothMaskIndex = InMaskIndex;
-						SelectedAsset = Asset;
+						const FClothLODDataCommon& LodData = ConcreteAsset->LodData[InAssetLod];
+
+						if(LodData.PointWeightMaps.IsValidIndex(InMaskIndex))
+						{
+							PaintingClothLODIndex = InAssetLod;
+							PaintingClothMaskIndex = InMaskIndex;
+							SelectedAsset = Asset;
+						}
 					}
-				}
 
-				break;
+					break;
+				}
 			}
 		}
 	}
@@ -263,13 +265,12 @@ bool FClothMeshPaintAdapter::InitializeVertexData()
 	MeshVertices.Empty();
 	MeshIndices.Empty();
 
-	if (SelectedAsset)
+	if (UClothingAssetCommon* const ConcreteAsset = Cast<UClothingAssetCommon>(SelectedAsset))
 	{
-		if(UDebugSkelMeshComponent* DebugComponent = Cast<UDebugSkelMeshComponent>(SkeletalMeshComponent))
+		if(const UDebugSkelMeshComponent* const DebugComponent = Cast<UDebugSkelMeshComponent>(SkeletalMeshComponent))
 		{
 			if(DebugComponent->SkinnedSelectedClothingPositions.Num() > 0)
 			{
-				UClothingAssetCommon* const ConcreteAsset = CastChecked<UClothingAssetCommon>(SelectedAsset);
 				if (ConcreteAsset->LodData.IsValidIndex(PaintingClothLODIndex))
 				{
 					const FClothLODDataCommon& LODData = ConcreteAsset->LodData[PaintingClothLODIndex];
@@ -334,9 +335,10 @@ FPointWeightMap* FClothMeshPaintAdapter::GetCurrentMask() const
 {
 	if(HasValidSelection())
 	{
-		UClothingAssetCommon* ConcreteAsset = CastChecked<UClothingAssetCommon>(SelectedAsset);
-
-		return &ConcreteAsset->LodData[PaintingClothLODIndex].PointWeightMaps[PaintingClothMaskIndex];
+		if(UClothingAssetCommon* const ConcreteAsset = Cast<UClothingAssetCommon>(SelectedAsset))
+		{
+			return &ConcreteAsset->LodData[PaintingClothLODIndex].PointWeightMaps[PaintingClothMaskIndex];
+		}
 	}
 
 	return nullptr;
@@ -344,10 +346,8 @@ FPointWeightMap* FClothMeshPaintAdapter::GetCurrentMask() const
 
 bool FClothMeshPaintAdapter::HasValidSelection() const
 {
-	if(SelectedAsset)
+	if(const UClothingAssetCommon* const ConcreteAsset = Cast<UClothingAssetCommon>(SelectedAsset))
 	{
-		UClothingAssetCommon* ConcreteAsset = Cast<UClothingAssetCommon>(SelectedAsset);
-
 		// Only valid selection if we have a valid asset, asset LOD and a mask
 		if(ConcreteAsset->LodData.IsValidIndex(PaintingClothLODIndex) &&
 			ConcreteAsset->LodData[PaintingClothLODIndex].PointWeightMaps.IsValidIndex(PaintingClothMaskIndex))
diff --git a/Engine/Source/Editor/ClothPainter/Private/ClothPainter.cpp b/Engine/Source/Editor/ClothPainter/Private/ClothPainter.cpp
--- a/Engine/Source/Editor/ClothPainter/Private/ClothPainter.cpp
+++ b/Engine/Source/Editor/ClothPainter/Private/ClothPainter.cpp
@@ -609,8 +609,14 @@ void FClothPainter::OnAssetSelectionChanged(UClothingAssetCommon* InNewSelectedA
 			{
 				Tool->OnMeshChanged();
 			}
+			return;
 		}
 	}
+
+	// Nothing to paint
+	SkeletalMeshComponent->SetMeshSectionVisibilityForCloth(SkeletalMeshComponent->SelectedClothingGuidForPainting, true);
+	SkeletalMeshComponent->SelectedClothingGuidForPainting = FGuid();
+	SkeletalMeshComponent->bDisableClothSimulation = false;
 }
 
 #undef LOCTEXT_NAMESPACE // "ClothPainter"
diff --git a/Engine/Source/Editor/ClothPainter/Private/ClothingPaintEditMode.cpp b/Engine/Source/Editor/ClothPainter/Private/ClothingPaintEditMode.cpp
--- a/Engine/Source/Editor/ClothPainter/Private/ClothingPaintEditMode.cpp
+++ b/Engine/Source/Editor/ClothPainter/Private/ClothingPaintEditMode.cpp
@@ -107,10 +107,12 @@ void FClothingPaintEditMode::Exit()
 			{
 				for(UClothingAssetBase* AssetBase : SkelMesh->GetMeshClothingAssets())
 				{
-					UClothingAssetCommon* ConcreteAsset = CastChecked<UClothingAssetCommon>(AssetBase);
-					constexpr bool bUpdateFixedVertData = true;
-					constexpr bool bInvalidateDerivedDataCache = true;
-					ConcreteAsset->ApplyParameterMasks(bUpdateFixedVertData, bInvalidateDerivedDataCache);
+					if (UClothingAssetCommon* const ConcreteAsset = Cast<UClothingAssetCommon>(AssetBase))
+					{
+						constexpr bool bUpdateFixedVertData = true;
+						constexpr bool bInvalidateDerivedDataCache = true;
+						ConcreteAsset->ApplyParameterMasks(bUpdateFixedVertData, bInvalidateDerivedDataCache);
+					}
 				}
 			}
 
diff --git a/Engine/Source/Editor/ClothPainter/Private/SClothAssetSelector.cpp b/Engine/Source/Editor/ClothPainter/Private/SClothAssetSelector.cpp
--- a/Engine/Source/Editor/ClothPainter/Private/SClothAssetSelector.cpp
+++ b/Engine/Source/Editor/ClothPainter/Private/SClothAssetSelector.cpp
@@ -67,6 +67,7 @@
 #include "Widgets/Images/SImage.h"
 #include "Widgets/Input/SButton.h"
 #include "Widgets/Input/SCheckBox.h"
+#include "Widgets/Input/SComboBox.h"
 #include "Widgets/Input/SComboButton.h"
 #include "Widgets/Input/SNumericEntryBox.h"
 #include "Widgets/Layout/SBox.h"
@@ -164,7 +165,7 @@ public:
 	{
 		if(Item.IsValid())
 		{
-			if(UClothingAssetCommon* Asset = Item->ClothingAsset.Get())
+			if(UClothingAssetBase* const Asset = Item->ClothingAsset.Get())
 			{
 				FText TrimText = FText::TrimPrecedingAndTrailing(InText);
 
@@ -173,9 +174,9 @@ public:
 					FName NewName(*TrimText.ToString());
 
 					// Check for an existing object, and if we find one build a unique name based on the request
-					if(UObject* ExistingObject = StaticFindObject(UClothingAssetCommon::StaticClass(), Asset->GetOuter(), *NewName.ToString()))
+					if(StaticFindObject(UClothingAssetBase::StaticClass(), Asset->GetOuter(), *NewName.ToString()))
 					{
-						NewName = MakeUniqueObjectName(Asset->GetOuter(), UClothingAssetCommon::StaticClass(), FName(*TrimText.ToString()));
+						NewName = MakeUniqueObjectName(Asset->GetOuter(), UClothingAssetBase::StaticClass(), FName(*TrimText.ToString()));
 					}
 
 					Asset->Rename(*NewName.ToString(), Asset->GetOuter());
@@ -261,7 +262,7 @@ private:
 			SourceSectionUserData.ClothingData = LODModel.Sections[SectionIndex].ClothingData;
 		};
 
-		if(UClothingAssetCommon* Asset = Item->ClothingAsset.Get())
+		if(UClothingAssetBase* const Asset = Item->ClothingAsset.Get())
 		{
 			if(USkeletalMesh* SkelMesh = Cast<USkeletalMesh>(Asset->GetOuter()))
 			{
@@ -307,7 +308,7 @@ private:
 			return;
 		}
 
-		if(UClothingAssetCommon* Asset = Item->ClothingAsset.Get())
+		if(UClothingAssetCommon* const Asset = Cast<UClothingAssetCommon>(Item->ClothingAsset.Get()))
 		{
 			const int32 NumLods = Asset->GetNumLods();
 
@@ -348,7 +349,7 @@ private:
 			return;
 		}
 
-		if (const UClothingAssetCommon* const ClothingAsset = Item->ClothingAsset.Get())
+		if (const UClothingAssetBase* const ClothingAsset = Item->ClothingAsset.Get())
 		{
 			ExportClothingAsset(ClothingAsset, ExportedType);
 		}
@@ -361,7 +362,7 @@ private:
 			return false;
 		}
 
-		if(UClothingAssetCommon* Asset = Item->ClothingAsset.Get())
+		if(const UClothingAssetCommon* const Asset = Cast<UClothingAssetCommon>(Item->ClothingAsset.Get()))
 		{
 			if(Asset->GetNumLods() > 1)
 			{
@@ -600,7 +601,7 @@ private:
 			{
 				if (FPointWeightMap* Mask = Item->GetMask())
 				{
-					if (UClothingAssetCommon* const ClothingAsset = AssetSelectorPtr.Pin()->GetSelectedAsset().Get())
+					if (UClothingAssetCommon* const ClothingAsset = Cast<UClothingAssetCommon>(AssetSelectorPtr.Pin()->GetSelectedAsset().Get()))
 					{
 						const int32 LOD = AssetSelectorPtr.Pin()->GetSelectedLod();
 						if (ClothingAsset->LodData.IsValidIndex(LOD))
@@ -656,7 +657,7 @@ private:
 	{
 		if (AssetSelectorPtr.IsValid() && Item.IsValid())
 		{
-			if (const UClothingAssetCommon* const ClothingAsset = AssetSelectorPtr.Pin()->GetSelectedAsset().Get())
+			if (const UClothingAssetCommon* const ClothingAsset = Cast<const UClothingAssetCommon>(AssetSelectorPtr.Pin()->GetSelectedAsset().Get()))
 			{
 				const int32 LOD = AssetSelectorPtr.Pin()->GetSelectedLod();
 				if (ClothingAsset->LodData.IsValidIndex(LOD))
@@ -685,13 +686,15 @@ private:
 	{
 		if (AssetSelectorPtr.IsValid())
 		{
-			UClothingAssetCommon* ClothingAsset = AssetSelectorPtr.Pin()->GetSelectedAsset().Get();
-			int32 LOD = AssetSelectorPtr.Pin()->GetSelectedLod();
-			FPointWeightMap* Mask = Item->GetMask();
+			if (UClothingAssetCommon* const ClothingAsset = Cast<UClothingAssetCommon>(AssetSelectorPtr.Pin()->GetSelectedAsset().Get()))
+			{
+				const int32 LOD = AssetSelectorPtr.Pin()->GetSelectedLod();
+				FPointWeightMap* const Mask = Item->GetMask();
 
-			TSharedRef<SWidget> Widget = SNew(SCopyVertexColorSettingsPanel, ClothingAsset, LOD, Mask);
+				const TSharedRef<SWidget> Widget = SNew(SCopyVertexColorSettingsPanel, ClothingAsset, LOD, Mask);
 
-			Builder.AddWidget(Widget, FText::GetEmpty(), true, false);
+				Builder.AddWidget(Widget, FText::GetEmpty(), true, false);
+			}
 		}
 	}
 
@@ -825,11 +828,9 @@ void SClothAssetSelector::Construct(const FArguments& InArgs, USkeletalMesh* InM
 				SAssignNew(AssetHeaderBox, SHorizontalBox)
 				+SHorizontalBox::Slot()
 				.VAlign(VAlign_Center)
-				.Padding(10.f, 0.f, 0.f, 0.f)
 				[
 					SNew(STextBlock)
 					.Text(LOCTEXT("AssetExpander_Title", "Clothing Data"))
-					.TransformPolicy(ETextTransformPolicy::ToUpper)
 					.TextStyle(FAppStyle::Get(), "DetailsView.CategoryTextStyle")
 					.Font(FAppStyle::Get().GetFontStyle("PropertyWindow.BoldFont"))
 				]
@@ -842,7 +843,7 @@ void SClothAssetSelector::Construct(const FArguments& InArgs, USkeletalMesh* InM
 					SNew(SPositiveActionButton)
 					.Text(LOCTEXT("CopyClothingFromMeshText_TEXT", "Add Clothing"))
 					.Icon(FAppStyle::Get().GetBrush("Icons.Plus"))
-					.ToolTipText(LOCTEXT("CopyClothingFromMeshText_TOOLTIP", "Copy Clothing from SkeletalMesh"))
+					.ToolTipText(LOCTEXT("CopyClothingFromMeshText_TOOLTIP", "Create clothing data from another asset."))
 					.OnGetMenuContent(this, &SClothAssetSelector::OnGenerateSkeletalMeshPickerForClothCopy)
 				]
 
@@ -852,12 +853,14 @@ void SClothAssetSelector::Construct(const FArguments& InArgs, USkeletalMesh* InM
 				.HAlign(HAlign_Right)
 				[
 					SNew(SComboButton)
+					.Visibility_Lambda([this]() -> EVisibility { return CanSelectLod() ? EVisibility::Visible : EVisibility::Collapsed; })
 					.ForegroundColor(FSlateColor::UseStyle())
 					.OnGetMenuContent(this, &SClothAssetSelector::OnGetLodMenu)
 					.HasDownArrow(true)
 					.ButtonContent()
 					[
 						SNew(STextBlock)
+						.Font(IDetailLayoutBuilder::GetDetailFont())
 						.Text(this, &SClothAssetSelector::GetLodButtonText)
 					]
 				]
@@ -883,6 +886,7 @@ void SClothAssetSelector::Construct(const FArguments& InArgs, USkeletalMesh* InM
 		.AutoHeight()
 		[
 			SNew(SExpandableArea)
+			.Visibility_Lambda([this]() -> EVisibility { return CanAddNewMask() ? EVisibility::Visible : EVisibility::Collapsed; })
 			.BorderImage(FAppStyle::Get().GetBrush("DetailsView.CategoryTop"))
 			.Padding(FMargin(0.f, 2.0f))
 			.HeaderContent()
@@ -890,11 +894,9 @@ void SClothAssetSelector::Construct(const FArguments& InArgs, USkeletalMesh* InM
 				SAssignNew(MaskHeaderBox, SHorizontalBox)
 				+SHorizontalBox::Slot()
 				.VAlign(VAlign_Center)
-				.Padding(10.f, 0.f, 0.f, 0.f)
 				[
 					SNew(STextBlock)
 					.Text(LOCTEXT("MaskExpander_Title", "Masks"))
-					.TransformPolicy(ETextTransformPolicy::ToUpper)
 					.TextStyle(FAppStyle::Get(), "DetailsView.CategoryTextStyle")
 					.Font(FAppStyle::Get().GetFontStyle("PropertyWindow.BoldFont"))
 				]
@@ -952,6 +954,7 @@ void SClothAssetSelector::Construct(const FArguments& InArgs, USkeletalMesh* InM
 		.Padding(FMargin(0.0f, 0.0f, 0.0f, 1.0f))
 		[
 			SNew(SExpandableArea)
+			.Visibility_Lambda([this]() -> EVisibility { return CanEditDeformer() ? EVisibility::Visible : EVisibility::Collapsed; })
 			.BorderImage(FAppStyle::Get().GetBrush("DetailsView.CategoryTop"))
 			.Padding(FMargin(0.f, 2.0f))
 			.HeaderContent()
@@ -959,11 +962,9 @@ void SClothAssetSelector::Construct(const FArguments& InArgs, USkeletalMesh* InM
 				SNew(SHorizontalBox)
 				+ SHorizontalBox::Slot()
 				.VAlign(VAlign_Center)
-				.Padding(10, 8, 0, 8)
 				[
 					SNew(STextBlock)
-					.Text(LOCTEXT("MeshSkinning_Title", "Mesh Skinning"))
-					.TransformPolicy(ETextTransformPolicy::ToUpper)
+					.Text(LOCTEXT("MeshProxyDeformer", "Mesh Proxy Deformer"))
 					.TextStyle(FAppStyle::Get(), "DetailsView.CategoryTextStyle")
 					.Font(FAppStyle::Get().GetFontStyle("PropertyWindow.BoldFont"))
 				]
@@ -988,6 +989,7 @@ void SClothAssetSelector::Construct(const FArguments& InArgs, USkeletalMesh* InM
 						SNew(STextBlock)
 						.Font(FAppStyle::Get().GetFontStyle("PropertyWindow.NormalFont"))
 						.Text(LOCTEXT("MultipleInfluences", "Use Multiple Influences"))
+						.ToolTipText(LOCTEXT("MultipleInfluencesToolTip", "When enabled, use multiple triangles from the simulation mesh to deform the render section, instead of a single one."))
 					]
 
 					+ SVerticalBox::Slot()
@@ -996,17 +998,17 @@ void SClothAssetSelector::Construct(const FArguments& InArgs, USkeletalMesh* InM
 					[
 						SNew(STextBlock)
 						.Font(FAppStyle::Get().GetFontStyle("PropertyWindow.NormalFont"))
-						.Text(LOCTEXT("CurrentRadius", "Kernel Radius"))
+						.Text(LOCTEXT("SearchRadius", "Search Radius"))
+						.ToolTipText(LOCTEXT("SearchRadiusToolTip", "The search radius around the render section vertices to look for all the influencing simulation triangles."))
 					]
 					
 					+ SVerticalBox::Slot()
-					.Padding(0.0f, 2.0f, 0.0f, 2.0f)
-					.AutoHeight()
+					.Padding(16.0f, 2.0f, 0.0f, 2.0f)
 					[
 						SNew(STextBlock)
-						.Font(IDetailLayoutBuilder::GetDetailFontBold())
-						.Text(LOCTEXT("SmoothTransition", "Smooth Transition From Skin to Cloth"))
-						.ShadowOffset(FVector2D(1, 1))
+						.Font(FAppStyle::Get().GetFontStyle("PropertyWindow.NormalFont"))
+						.Text(LOCTEXT("SmoothTransition", "Smooth Transition"))
+						.ToolTipText(LOCTEXT("SmoothTransitionToolTip", "Whether to enable smooth transitions between kinematic and dynamic simulation triangles deforming the render section."))
 					]
 				]
 
@@ -1042,8 +1044,8 @@ void SClothAssetSelector::Construct(const FArguments& InArgs, USkeletalMesh* InM
 					]
 					
 					+ SVerticalBox::Slot()
-					.Padding(0.0f, 2.0f, 0.0f, 2.0f)
-					.AutoHeight()
+					.VAlign(VAlign_Center)
+					.Padding(12.0f, 4.0f, 0.0f, 4.0f)
 					[
 						SNew(SCheckBox)
 						.IsChecked(this, &SClothAssetSelector::GetCurrentSmoothTransition)
@@ -1051,7 +1053,6 @@ void SClothAssetSelector::Construct(const FArguments& InArgs, USkeletalMesh* InM
 						.IsEnabled(this, &SClothAssetSelector::IsValidClothLodSelected)
 					]
 				]
-
 			]
 		]
 	];
@@ -1060,22 +1061,20 @@ void SClothAssetSelector::Construct(const FArguments& InArgs, USkeletalMesh* InM
 	RefreshMaskList();
 }
 
-
 TOptional<float> SClothAssetSelector::GetCurrentKernelRadius() const
 {
-	UClothingAssetCommon* Asset = SelectedAsset.Get();
+	const UClothingAssetCommon* const Asset = Cast<UClothingAssetCommon>(SelectedAsset.Get());
 	if (Asset && Asset->IsValidLod(SelectedLod))
 	{
 		const FClothLODDataCommon& LodData = Asset->LodData[SelectedLod];
 		return LodData.SkinningKernelRadius;
 	}
-
 	return TOptional<float>();
 }
 
 void SClothAssetSelector::OnCurrentKernelRadiusChanged(float InValue)
 {
-	UClothingAssetCommon* Asset = SelectedAsset.Get();
+	UClothingAssetCommon* const Asset = Cast<UClothingAssetCommon>(SelectedAsset.Get());
 	if (Asset && Asset->IsValidLod(SelectedLod))
 	{
 		FClothLODDataCommon& LodData = Asset->LodData[SelectedLod];
@@ -1085,7 +1084,7 @@ void SClothAssetSelector::OnCurrentKernelRadiusChanged(float InValue)
 
 void SClothAssetSelector::OnCurrentKernelRadiusCommitted(float InValue, ETextCommit::Type CommitType)
 {
-	UClothingAssetCommon* Asset = SelectedAsset.Get();
+	UClothingAssetCommon* const Asset = Cast<UClothingAssetCommon>(SelectedAsset.Get());
 	if (Asset && Asset->IsValidLod(SelectedLod))
 	{
 		FClothLODDataCommon& LodData = Asset->LodData[SelectedLod];
@@ -1107,7 +1106,7 @@ bool SClothAssetSelector::CurrentKernelRadiusIsEnabled() const
 
 ECheckBoxState SClothAssetSelector::GetCurrentUseMultipleInfluences() const
 {
-	UClothingAssetCommon* Asset = SelectedAsset.Get();
+	UClothingAssetCommon* const Asset = Cast<UClothingAssetCommon>(SelectedAsset.Get());
 	if (Asset && Asset->IsValidLod(SelectedLod))
 	{
 		const FClothLODDataCommon& LodData = Asset->LodData[SelectedLod];
@@ -1124,7 +1123,7 @@ void SClothAssetSelector::OnCurrentUseMultipleInfluencesChanged(ECheckBoxState I
 		return;
 	}
 
-	UClothingAssetCommon* Asset = SelectedAsset.Get();
+	UClothingAssetCommon* const Asset = Cast<UClothingAssetCommon>(SelectedAsset.Get());
 	if (Asset && Asset->IsValidLod(SelectedLod))
 	{
 		FClothLODDataCommon& LodData = Asset->LodData[SelectedLod];
@@ -1142,7 +1141,7 @@ void SClothAssetSelector::OnCurrentUseMultipleInfluencesChanged(ECheckBoxState I
 
 ECheckBoxState SClothAssetSelector::GetCurrentSmoothTransition() const
 {
-	UClothingAssetCommon* Asset = SelectedAsset.Get();
+	UClothingAssetCommon* const Asset = Cast<UClothingAssetCommon>(SelectedAsset.Get());
 	if (Asset && Asset->IsValidLod(SelectedLod))
 	{
 		const FClothLODDataCommon& LodData = Asset->LodData[SelectedLod];
@@ -1160,8 +1159,8 @@ void SClothAssetSelector::OnCurrentSmoothTransitionChanged(ECheckBoxState InValu
 		return;
 	}
 
-	UClothingAssetCommon* Asset = SelectedAsset.Get();
-	if (Asset && Asset->IsValidLod(SelectedLod))
+	UClothingAssetCommon* const Asset = Cast<UClothingAssetCommon>(SelectedAsset.Get());
+	if (Asset &&Asset->IsValidLod(SelectedLod))
 	{
 		FClothLODDataCommon& LodData = Asset->LodData[SelectedLod];
 		LodData.bSmoothTransition = (InValue == ECheckBoxState::Checked);
@@ -1177,12 +1176,11 @@ void SClothAssetSelector::OnCurrentSmoothTransitionChanged(ECheckBoxState InValu
 
 bool SClothAssetSelector::IsValidClothLodSelected() const
 {
-	UClothingAssetCommon* Asset = SelectedAsset.Get();
+	const UClothingAssetCommon* const Asset = Cast<UClothingAssetCommon>(SelectedAsset.Get());
 	return (Asset && Asset->IsValidLod(SelectedLod));
 }
 
-
-TWeakObjectPtr<UClothingAssetCommon> SClothAssetSelector::GetSelectedAsset() const
+TWeakObjectPtr<UClothingAssetBase> SClothAssetSelector::GetSelectedAsset() const
 {
 	return SelectedAsset;
 	
@@ -1239,12 +1237,28 @@ void SClothAssetSelector::OnCopyClothingAssetSelected(const FAssetData& AssetDat
 
 TSharedRef<SWidget> SClothAssetSelector::OnGenerateSkeletalMeshPickerForClothCopy()
 {
-	FContentBrowserModule& ContentBrowserModule = FModuleManager::Get().LoadModuleChecked<FContentBrowserModule>(TEXT("ContentBrowser"));
+	FContentBrowserModule& ContentBrowserModule = FModuleManager::Get().LoadModuleChecked<FContentBrowserModule>("ContentBrowser");
+	FClothingSystemEditorInterfaceModule& ClothingEditorModule = FModuleManager::LoadModuleChecked<FClothingSystemEditorInterfaceModule>("ClothingSystemEditorInterface");
 
 	FAssetPickerConfig AssetPickerConfig;
-	AssetPickerConfig.Filter.ClassPaths.Add(USkeletalMesh::StaticClass()->GetClassPathName());
 	AssetPickerConfig.OnAssetSelected = FOnAssetSelected::CreateSP(this, &SClothAssetSelector::OnCopyClothingAssetSelected);
-	AssetPickerConfig.bAllowNullSelection = true;
+	AssetPickerConfig.OnShouldFilterAsset = FOnShouldFilterAsset::CreateLambda(
+		[this, &ClothingEditorModule](const FAssetData& AssetData)
+		{
+			const FSoftClassPath SoftClassPath(AssetData.AssetClassPath.ToString());
+			if (UClass* const ResolvedClass = SoftClassPath.ResolveClass())
+			{
+				for (const UClothingAssetFactoryBase* const ClothingAssetFactory : ClothingEditorModule.GetClothingAssetFactories())
+				{
+					if (ResolvedClass->IsChildOf(ClothingAssetFactory->GetSupportedSourceAssetType()))
+					{
+						return false;
+					}
+				}
+			}
+			return true;
+		});
+	AssetPickerConfig.bAllowNullSelection = false;
 	AssetPickerConfig.InitialAssetViewType = EAssetViewType::List;
 	AssetPickerConfig.bFocusSearchBoxWhenOpened = true;
 	AssetPickerConfig.bShowBottomToolbar = false;
@@ -1272,13 +1286,23 @@ EVisibility SClothAssetSelector::GetMaskHeaderButtonTextVisibility() const
 	return bShow ? EVisibility::HitTestInvisible : EVisibility::Collapsed;
 }
 
+bool SClothAssetSelector::CanSelectLod() const
+{
+	return Cast<const UClothingAssetCommon>(SelectedAsset.Get()) != nullptr;
+}
+
+bool SClothAssetSelector::CanEditDeformer() const
+{
+	return Cast<const UClothingAssetCommon>(SelectedAsset.Get()) != nullptr;
+}
+
 TSharedRef<SWidget> SClothAssetSelector::OnGetLodMenu()
 {
 	FMenuBuilder Builder(true, nullptr);
 
 	int32 NumLods = 0;
 
-	if(UClothingAssetCommon* CurrAsset = SelectedAsset.Get())
+	if(const UClothingAssetCommon* const CurrAsset = Cast<UClothingAssetCommon>(SelectedAsset.Get()))
 	{
 		NumLods = CurrAsset->GetNumLods();
 	}
@@ -1316,7 +1340,7 @@ FText SClothAssetSelector::GetLodButtonText() const
 
 TSharedRef<ITableRow> SClothAssetSelector::OnGenerateWidgetForClothingAssetItem(TSharedPtr<FClothingAssetListItem> InItem, const TSharedRef<STableViewBase>& OwnerTable)
 {
-	if(UClothingAssetCommon* Asset = InItem->ClothingAsset.Get())
+	if(InItem.IsValid() && InItem->ClothingAsset.Get())
 	{
 		return SNew(SAssetListRow, OwnerTable, InItem)
 			.OnInvalidateList(this, &SClothAssetSelector::OnRefresh);
@@ -1366,7 +1390,7 @@ void SClothAssetSelector::OnMaskSelectionChanged(TSharedPtr<FClothingMaskListIte
 
 FReply SClothAssetSelector::AddNewMask()
 {
-	if(UClothingAssetCommon* Asset = SelectedAsset.Get())
+	if(UClothingAssetCommon* const Asset = Cast<UClothingAssetCommon>(SelectedAsset.Get()))
 	{
 		if(Asset->LodData.IsValidIndex(SelectedLod))
 		{
@@ -1390,7 +1414,7 @@ FReply SClothAssetSelector::AddNewMask()
 
 bool SClothAssetSelector::CanAddNewMask() const
 {
-	return SelectedAsset.Get() != nullptr;
+	return Cast<const UClothingAssetCommon>(SelectedAsset.Get()) != nullptr;
 }
 
 void SClothAssetSelector::OnRefresh()
@@ -1401,7 +1425,7 @@ void SClothAssetSelector::OnRefresh()
 
 void SClothAssetSelector::RefreshAssetList()
 {
-	UClothingAssetCommon* CurrSelectedAsset = nullptr;
+	UClothingAssetBase* CurrSelectedAsset = nullptr;
 	int32 SelectedItem = INDEX_NONE;
 
 	if(AssetList.IsValid())
@@ -1419,15 +1443,13 @@ void SClothAssetSelector::RefreshAssetList()
 
 	for(UClothingAssetBase* Asset : Mesh->GetMeshClothingAssets())
 	{
-		UClothingAssetCommon* ConcreteAsset = Cast<UClothingAssetCommon>(Asset);
-
 		TSharedPtr<FClothingAssetListItem> Entry = MakeShareable(new FClothingAssetListItem);
 
-		Entry->ClothingAsset = ConcreteAsset;
+		Entry->ClothingAsset = Asset;
 
 		AssetListItems.Add(Entry);
 
-		if(ConcreteAsset == CurrSelectedAsset)
+		if(Asset == CurrSelectedAsset)
 		{
 			SelectedItem = AssetListItems.Num() - 1;
 		}
@@ -1471,7 +1493,7 @@ void SClothAssetSelector::RefreshMaskList()
 
 	MaskListItems.Empty();
 
-	UClothingAssetCommon* Asset = SelectedAsset.Get();
+	UClothingAssetCommon* const Asset = Cast<UClothingAssetCommon>(SelectedAsset.Get());
 	if(Asset && Asset->IsValidLod(SelectedLod))
 	{
 		const FClothLODDataCommon& LodData = Asset->LodData[SelectedLod];
@@ -1480,7 +1502,7 @@ void SClothAssetSelector::RefreshMaskList()
 		for(int32 Index = 0; Index < NumMasks; ++Index)
 		{
 			TSharedPtr<FClothingMaskListItem> NewItem = MakeShareable(new FClothingMaskListItem);
-			NewItem->ClothingAsset = SelectedAsset;
+			NewItem->ClothingAsset = Asset;
 			NewItem->LodIndex = SelectedLod;
 			NewItem->MaskIndex = Index;
 			MaskListItems.Add(NewItem);
@@ -1519,14 +1541,14 @@ void SClothAssetSelector::OnClothingLodSelected(int32 InNewLod)
 		//ClothPainterSettings->OnAssetSelectionChanged.Broadcast(SelectedAsset.Get(), SelectedLod, SelectedMask);
 	}
 
-	if(SelectedAsset.IsValid())
+	if(const UClothingAssetCommon* const Asset = Cast<UClothingAssetCommon>(SelectedAsset.Get()))
 	{
 		SetSelectedLod(InNewLod);
 
 		int32 NewMaskSelection = INDEX_NONE;
-		if(SelectedAsset->LodData.IsValidIndex(SelectedLod))
+		if(Asset->LodData.IsValidIndex(SelectedLod))
 		{
-			const FClothLODDataCommon& LodData = SelectedAsset->LodData[SelectedLod];
+			const FClothLODDataCommon& LodData = Asset->LodData[SelectedLod];
 
 			if(LodData.PointWeightMaps.Num() > 0)
 			{
@@ -1538,36 +1560,34 @@ void SClothAssetSelector::OnClothingLodSelected(int32 InNewLod)
 	}
 }
 
-void SClothAssetSelector::SetSelectedAsset(TWeakObjectPtr<UClothingAssetCommon> InSelectedAsset)
+void SClothAssetSelector::SetSelectedAsset(TWeakObjectPtr<UClothingAssetBase> InSelectedAsset)
 {
 	SelectedAsset = InSelectedAsset;
 
 	RefreshMaskList();
 
-	if(UClothingAssetCommon* NewAsset = SelectedAsset.Get())
+	const UClothingAssetCommon* const NewAsset = Cast<const UClothingAssetCommon>(SelectedAsset.Get());
+	if(NewAsset && NewAsset->GetNumLods() > 0)
 	{
-		if(NewAsset->GetNumLods() > 0)
-		{
-			SetSelectedLod(0);
+		SetSelectedLod(0);
 
-			const FClothLODDataCommon& LodData = NewAsset->LodData[SelectedLod];
-			if(LodData.PointWeightMaps.Num() > 0)
-			{
-				SetSelectedMask(0);
-			}
-			else
-			{
-				SetSelectedMask(INDEX_NONE);
-			}
+		const FClothLODDataCommon& LodData = NewAsset->LodData[SelectedLod];
+		if(LodData.PointWeightMaps.Num() > 0)
+		{
+			SetSelectedMask(0);
 		}
 		else
 		{
-			SetSelectedLod(INDEX_NONE);
 			SetSelectedMask(INDEX_NONE);
 		}
-
-		OnSelectionChanged.ExecuteIfBound(SelectedAsset, SelectedLod, SelectedMask);
 	}
+	else
+	{
+		SetSelectedLod(INDEX_NONE);
+		SetSelectedMask(INDEX_NONE);
+	}
+
+	OnSelectionChanged.ExecuteIfBound(SelectedAsset, SelectedLod, SelectedMask);
 }
 
 void SClothAssetSelector::SetSelectedLod(int32 InLodIndex, bool bRefreshMasks /*= true*/)
diff --git a/Engine/Source/Editor/ClothPainter/Private/SClothAssetSelector.h b/Engine/Source/Editor/ClothPainter/Private/SClothAssetSelector.h
--- a/Engine/Source/Editor/ClothPainter/Private/SClothAssetSelector.h
+++ b/Engine/Source/Editor/ClothPainter/Private/SClothAssetSelector.h
@@ -31,6 +31,7 @@ class SHorizontalBox;
 class STableViewBase;
 class SWidget;
 class UClothingAssetCommon;
+class UClothingAssetBase;
 class USkeletalMesh;
 struct FAssetData;
 struct FClothPhysicalMeshData;
@@ -38,7 +39,7 @@ struct FPointWeightMap;
 
 struct FClothingAssetListItem
 {
-	TWeakObjectPtr<UClothingAssetCommon> ClothingAsset;
+	TWeakObjectPtr<UClothingAssetBase> ClothingAsset;
 };
 
 struct FClothingMaskListItem
@@ -60,7 +61,7 @@ struct FClothingMaskListItem
 typedef SListView<TSharedPtr<FClothingAssetListItem>> SAssetList;
 typedef SListView<TSharedPtr<FClothingMaskListItem>> SMaskList;
 
-DECLARE_DELEGATE_ThreeParams(FOnClothAssetSelectionChanged, TWeakObjectPtr<UClothingAssetCommon>, int32, int32);
+DECLARE_DELEGATE_ThreeParams(FOnClothAssetSelectionChanged, TWeakObjectPtr<UClothingAssetBase>, int32, int32);
 
 class SClothAssetSelector : public SCompoundWidget, public FEditorUndoClient
 {
@@ -74,7 +75,7 @@ public:
 
 	void Construct(const FArguments& InArgs, USkeletalMesh* InMesh);
 
-	TWeakObjectPtr<UClothingAssetCommon> GetSelectedAsset() const;
+	TWeakObjectPtr<UClothingAssetBase> GetSelectedAsset() const;
 	int32 GetSelectedLod() const;
 	int32 GetSelectedMask() const;
 
@@ -127,7 +128,7 @@ protected:
 	void OnClothingLodSelected(int32 InNewLod);
 
 	// Setters for the list selections so we can handle list selections changing properly
-	void SetSelectedAsset(TWeakObjectPtr<UClothingAssetCommon> InSelectedAsset);
+	void SetSelectedAsset(TWeakObjectPtr<UClothingAssetBase> InSelectedAsset);
 	void SetSelectedLod(int32 InLodIndex, bool bRefreshMasks = true);
 	void SetSelectedMask(int32 InMaskIndex);
 
@@ -144,7 +145,7 @@ protected:
 	TArray<TSharedPtr<FClothingMaskListItem>> MaskListItems;
 
 	// Currently selected clothing asset, Lod Index and Mask index
-	TWeakObjectPtr<UClothingAssetCommon> SelectedAsset;
+	TWeakObjectPtr<UClothingAssetBase> SelectedAsset;
 	int32 SelectedLod;
 	int32 SelectedMask;
 
@@ -152,4 +153,8 @@ protected:
 
 	// Handle for mesh event callback when clothing changes.
 	FDelegateHandle MeshClothingChangedHandle;
+
+private:
+	bool CanSelectLod() const;
+	bool CanEditDeformer() const;
 };
diff --git a/Engine/Source/Editor/ClothPainter/Private/SClothPaintTab.cpp b/Engine/Source/Editor/ClothPainter/Private/SClothPaintTab.cpp
--- a/Engine/Source/Editor/ClothPainter/Private/SClothPaintTab.cpp
+++ b/Engine/Source/Editor/ClothPainter/Private/SClothPaintTab.cpp
@@ -49,6 +49,7 @@ SClothPaintTab::~SClothPaintTab()
 void SClothPaintTab::Construct(const FArguments& InArgs)
 {
 	// Detail view for UClothingAssetCommon
+	const IConsoleVariable* const EnableCardLayoutVarriable = IConsoleManager::Get().FindConsoleVariable(TEXT("DetailsPanel.Style.EnableCardLayout"));
 	FPropertyEditorModule& EditModule = FModuleManager::Get().GetModuleChecked<FPropertyEditorModule>("PropertyEditor");
 
 	FDetailsViewArgs DetailsViewArgs;
@@ -58,6 +59,9 @@ void SClothPaintTab::Construct(const FArguments& InArgs)
 	DetailsViewArgs.DefaultsOnlyVisibility = EEditDefaultsOnlyNodeVisibility::Automatic;
 	DetailsViewArgs.bShowOptions = false;
 	DetailsViewArgs.bAllowMultipleTopLevelObjects = true;
+	DetailsViewArgs.StyleKey = EnableCardLayoutVarriable && EnableCardLayoutVarriable->GetBool() ?
+		MakeShared<FDetailsViewStyleKey>(FDetailsViewStyleKeys::Card()) :
+		MakeShared<FDetailsViewStyleKey>(FDetailsViewStyleKeys::Classic());
 
 	DetailsView = EditModule.CreateDetailView(DetailsViewArgs);
 	
@@ -131,12 +135,8 @@ void SClothPaintTab::EnterPaintMode()
 
 	if(SelectorWidget.IsValid())
 	{
-		TWeakObjectPtr<UClothingAssetCommon> WeakAsset = SelectorWidget->GetSelectedAsset();
-
-		if(WeakAsset.Get())
-		{
-			ClothPainter->OnAssetSelectionChanged(WeakAsset.Get(), SelectorWidget->GetSelectedLod(), SelectorWidget->GetSelectedMask());
-		}
+		UClothingAssetCommon* const Asset = Cast<UClothingAssetCommon>(SelectorWidget->GetSelectedAsset().Get());
+		ClothPainter->OnAssetSelectionChanged(Asset, SelectorWidget->GetSelectedLod(), SelectorWidget->GetSelectedMask());
 	}
 }
 
@@ -146,7 +146,7 @@ void SClothPaintTab::ExitPaintMode()
 	ModeWidget = nullptr;
 }
 
-void SClothPaintTab::OnAssetSelectionChanged(TWeakObjectPtr<UClothingAssetCommon> InAssetPtr, int32 InLodIndex, int32 InMaskIndex)
+void SClothPaintTab::OnAssetSelectionChanged(TWeakObjectPtr<UClothingAssetBase> InAssetPtr, int32 InLodIndex, int32 InMaskIndex)
 {
 	const ISkeletalMeshEditor* SkeletalMeshEditor = GetSkeletalMeshEditor();
 	if (!SkeletalMeshEditor)
@@ -157,20 +157,22 @@ void SClothPaintTab::OnAssetSelectionChanged(TWeakObjectPtr<UClothingAssetCommon
 	FClothingPaintEditMode* PaintMode = (FClothingPaintEditMode*)SkeletalMeshEditor->GetEditorModeManager().GetActiveMode(PaintModeID);
 	if(PaintMode)
 	{
-		if(FClothPainter* ClothPainter = static_cast<FClothPainter*>(PaintMode->GetMeshPainter()))
+		if(FClothPainter* const ClothPainter = static_cast<FClothPainter*>(PaintMode->GetMeshPainter()))
 		{
-			ClothPainter->OnAssetSelectionChanged(InAssetPtr.Get(), InLodIndex, InMaskIndex);
+			ClothPainter->OnAssetSelectionChanged(Cast<UClothingAssetCommon>(InAssetPtr.Get()), InLodIndex, InMaskIndex);
 		}
 	}
 
-	if(UClothingAssetCommon* Asset = InAssetPtr.Get())
+	TArray<UObject*> Objects;
+	if(UClothingAssetBase* const Asset = InAssetPtr.Get())
 	{
-		TArray<UObject*> Objects;
-
 		Objects.Add(Asset);
-
-		DetailsView->SetObjects(Objects, true);
 	}
+	else
+	{
+		Objects.Add(UClothingAssetCommon::StaticClass()->GetDefaultObject());  // Add CDO by default
+	}
+	DetailsView->SetObjects(Objects, true);
 }
 
 bool SClothPaintTab::IsAssetDetailsPanelEnabled()
diff --git a/Engine/Source/Editor/ClothPainter/Private/SClothPaintTab.h b/Engine/Source/Editor/ClothPainter/Private/SClothPaintTab.h
--- a/Engine/Source/Editor/ClothPainter/Private/SClothPaintTab.h
+++ b/Engine/Source/Editor/ClothPainter/Private/SClothPaintTab.h
@@ -16,7 +16,7 @@ class ISkeletalMeshEditor;
 class SClothAssetSelector;
 class SClothPaintWidget;
 class SScrollBox;
-class UClothingAssetCommon;
+class UClothingAssetBase;
 struct FGeometry;
 
 class SClothPaintTab : public SCompoundWidget
@@ -41,7 +41,7 @@ public:
 protected:
 
 	/** Called from the selector when the asset selection changes (Asset, LOD, Mask) */
-	UE_API void OnAssetSelectionChanged(TWeakObjectPtr<UClothingAssetCommon> InAssetPtr, int32 InLodIndex, int32 InMaskIndex);
+	UE_API void OnAssetSelectionChanged(TWeakObjectPtr<UClothingAssetBase> InAssetPtr, int32 InLodIndex, int32 InMaskIndex);
 
 	/** Whether or not the asset config section is enabled for editing */
 	UE_API bool IsAssetDetailsPanelEnabled();
diff --git a/Engine/Source/Editor/Persona/Private/PersonaMeshDetails.cpp b/Engine/Source/Editor/Persona/Private/PersonaMeshDetails.cpp
--- a/Engine/Source/Editor/Persona/Private/PersonaMeshDetails.cpp
+++ b/Engine/Source/Editor/Persona/Private/PersonaMeshDetails.cpp
@@ -6765,8 +6765,8 @@ void FPersonaMeshDetails::OnGenerateElementForClothingAsset( TSharedRef<IPropert
 	[
 		MakeClothingDetailsWidget(ElementIndex)
 	];	
-	
-	// Properties are now inside UClothingAssetCommon, so we just add a new inspector and handle everything through that
+
+	// Properties are now inside the clothing data, so we just add a new inspector and handle everything through that
 	FDetailWidgetRow& ClothPropRow = ChildrenBuilder.AddCustomRow(LOCTEXT("ClothingAsset_Search_Properties", "Properties"));
 
 	TSharedPtr<SKismetInspector> Inspector = nullptr;
@@ -6805,31 +6805,25 @@ TSharedRef<SUniformGridPanel> FPersonaMeshDetails::MakeClothingDetailsWidget(int
 	TSharedRef<SUniformGridPanel> Grid = SNew(SUniformGridPanel).SlotPadding(2.0f);
 
 	USkeletalMesh* SkelMesh = GetPersonaToolkit()->GetMesh();
-	UClothingAssetBase* ClothingAsset = SkelMesh->GetMeshClothingAssets()[AssetIndex];
-	if (!ClothingAsset)
-	{
-		return Grid;
-	}
-
 
-	const FSlateFontInfo DetailFontInfo = IDetailLayoutBuilder::GetDetailFont();
-	const int32 NumLODs = ClothingAsset->GetNumLods();
-	int32 RowNumber = 0;
-	for(int32 LODIndex=0; LODIndex < NumLODs; LODIndex++)
+	if (UClothingAssetCommon* const ClothingAsset = Cast<UClothingAssetCommon>(SkelMesh->GetMeshClothingAssets()[AssetIndex]))
 	{
-		Grid->AddSlot(0, RowNumber) // x, y
-		.HAlign(HAlign_Left)
-		[
-			SNew(STextBlock)
-			.Font(DetailFontInfo)
-			.Text(FText::Format(LOCTEXT("LODIndexLabel", "LOD {0}"), FText::AsNumber(LODIndex)))			
-		];
+		const FSlateFontInfo DetailFontInfo = IDetailLayoutBuilder::GetDetailFont();
+		const int32 NumLODs = ClothingAsset->GetNumLods();
+		int32 RowNumber = 0;
+		for (int32 LODIndex = 0; LODIndex < NumLODs; ++LODIndex)
+		{
+			Grid->AddSlot(0, RowNumber) // x, y
+				.HAlign(HAlign_Left)
+				[
+					SNew(STextBlock)
+					.Font(DetailFontInfo)
+					.Text(FText::Format(LOCTEXT("LODIndexLabel", "LOD {0}"), FText::AsNumber(LODIndex)))
+				];
 
-		RowNumber++;
+			++RowNumber;
 
-		if (UClothingAssetCommon* Asset = Cast<UClothingAssetCommon>(ClothingAsset))
-		{
-			const FClothLODDataCommon& LodData = Asset->LodData[LODIndex];
+			const FClothLODDataCommon& LodData = ClothingAsset->LodData[LODIndex];
 			const FClothPhysicalMeshData& PhysMeshData = LodData.PhysicalMeshData;
 
 			Grid->AddSlot(0, RowNumber)
@@ -6837,7 +6831,7 @@ TSharedRef<SUniformGridPanel> FPersonaMeshDetails::MakeClothingDetailsWidget(int
 				[
 					SNew(STextBlock)
 					.Font(DetailFontInfo)
-				.Text(LOCTEXT("SimulVertexCount", "Simul Verts"))
+					.Text(LOCTEXT("SimulVertexCount", "Simul Verts"))
 				];
 
 			Grid->AddSlot(0, RowNumber + 1)
@@ -6845,7 +6839,7 @@ TSharedRef<SUniformGridPanel> FPersonaMeshDetails::MakeClothingDetailsWidget(int
 				[
 					SNew(STextBlock)
 					.Font(DetailFontInfo)
-				.Text(FText::AsNumber(PhysMeshData.Vertices.Num() - PhysMeshData.NumFixedVerts))
+					.Text(FText::AsNumber(PhysMeshData.Vertices.Num() - PhysMeshData.NumFixedVerts))
 				];
 
 			Grid->AddSlot(1, RowNumber)
@@ -6853,7 +6847,7 @@ TSharedRef<SUniformGridPanel> FPersonaMeshDetails::MakeClothingDetailsWidget(int
 				[
 					SNew(STextBlock)
 					.Font(DetailFontInfo)
-				.Text(LOCTEXT("FixedVertexCount", "Fixed Verts"))
+					.Text(LOCTEXT("FixedVertexCount", "Fixed Verts"))
 				];
 
 			Grid->AddSlot(1, RowNumber + 1)
@@ -6861,7 +6855,7 @@ TSharedRef<SUniformGridPanel> FPersonaMeshDetails::MakeClothingDetailsWidget(int
 				[
 					SNew(STextBlock)
 					.Font(DetailFontInfo)
-				.Text(FText::AsNumber(PhysMeshData.NumFixedVerts))
+					.Text(FText::AsNumber(PhysMeshData.NumFixedVerts))
 				];
 
 			Grid->AddSlot(2, RowNumber)
@@ -6869,7 +6863,7 @@ TSharedRef<SUniformGridPanel> FPersonaMeshDetails::MakeClothingDetailsWidget(int
 				[
 					SNew(STextBlock)
 					.Font(DetailFontInfo)
-				.Text(LOCTEXT("TriangleCount", "Sim Triangles"))
+					.Text(LOCTEXT("TriangleCount", "Sim Triangles"))
 				];
 
 			Grid->AddSlot(2, RowNumber + 1)
@@ -6877,7 +6871,7 @@ TSharedRef<SUniformGridPanel> FPersonaMeshDetails::MakeClothingDetailsWidget(int
 				[
 					SNew(STextBlock)
 					.Font(DetailFontInfo)
-				.Text(FText::AsNumber(PhysMeshData.Indices.Num() / 3))
+					.Text(FText::AsNumber(PhysMeshData.Indices.Num() / 3))
 				];
 
 			Grid->AddSlot(3, RowNumber)
@@ -6885,7 +6879,7 @@ TSharedRef<SUniformGridPanel> FPersonaMeshDetails::MakeClothingDetailsWidget(int
 				[
 					SNew(STextBlock)
 					.Font(DetailFontInfo)
-				.Text(LOCTEXT("NumUsedBones", "Bones"))
+					.Text(LOCTEXT("NumUsedBones", "Bones"))
 				];
 
 			Grid->AddSlot(3, RowNumber + 1)
@@ -6893,16 +6887,11 @@ TSharedRef<SUniformGridPanel> FPersonaMeshDetails::MakeClothingDetailsWidget(int
 				[
 					SNew(STextBlock)
 					.Font(DetailFontInfo)
-				.Text(FText::AsNumber(PhysMeshData.MaxBoneWeights))
+					.Text(FText::AsNumber(PhysMeshData.MaxBoneWeights))
 				];
 
 			RowNumber += 2;
 		}
-		else
-		{
-			// Unsupported asset type
-			check(false);
-		}
 	}
 
 	return Grid;
@@ -6966,15 +6955,26 @@ void FPersonaMeshDetails::UpdateClothingEntries()
 	const int32 NumClothingAssets = Mesh->GetMeshClothingAssets().Num();
 	for(int32 Idx = 0; Idx < NumClothingAssets; ++Idx)
 	{
-		if (UClothingAssetBase* ClothingAsset = Mesh->GetMeshClothingAssets()[Idx])
+		if (UClothingAssetBase* const ClothingAsset = Mesh->GetMeshClothingAssets()[Idx])
 		{
-			const int32 NumAssetLods = ClothingAsset->GetNumLods();
-			for (int32 AssetLodIndex = 0; AssetLodIndex < NumAssetLods; ++AssetLodIndex)
+			if (UClothingAssetCommon* const ClothingAssetCommon = Cast<UClothingAssetCommon>(ClothingAsset))
+			{
+				const int32 NumAssetLods = ClothingAssetCommon->GetNumLods();
+				for (int32 AssetLodIndex = 0; AssetLodIndex < NumAssetLods; ++AssetLodIndex)
+				{
+					TSharedPtr<FClothingEntry> NewEntry = MakeShared<FClothingEntry>();
+					NewEntry->Asset = ClothingAsset;
+					NewEntry->AssetIndex = Idx;
+					NewEntry->AssetLodIndex = AssetLodIndex;
+					NewClothingAssetEntries.Add(NewEntry);
+				}
+			}
+			else
 			{
 				TSharedPtr<FClothingEntry> NewEntry = MakeShared<FClothingEntry>();
 				NewEntry->Asset = ClothingAsset;
 				NewEntry->AssetIndex = Idx;
-				NewEntry->AssetLodIndex = AssetLodIndex;
+				NewEntry->AssetLodIndex = INDEX_NONE;
 				NewClothingAssetEntries.Add(NewEntry);
 			}
 		}
@@ -7000,12 +7000,14 @@ void FPersonaMeshDetails::OnClothingComboBoxOpening()
 
 TSharedRef<SWidget> FPersonaMeshDetails::OnGenerateWidgetForClothingEntry(TSharedPtr<FClothingEntry> InEntry)
 {
-	UClothingAssetCommon* Asset = Cast<UClothingAssetCommon>(InEntry->Asset.Get());
+	const UClothingAssetBase* const Asset = Cast<UClothingAssetBase>(InEntry->Asset.Get());
 
 	FText EntryText;
 	if(Asset)
 	{
-		EntryText = FText::Format(LOCTEXT("ClothingAssetEntry_Name", "{0} - LOD{1}"), FText::FromString(Asset->GetName()), FText::AsNumber(InEntry->AssetLodIndex));
+		EntryText = (InEntry->AssetLodIndex != INDEX_NONE) ?
+			FText::Format(LOCTEXT("ClothingAssetEntry_Name", "{0} - LOD{1}"), FText::FromString(Asset->GetName()), FText::AsNumber(InEntry->AssetLodIndex)) : 
+			FText::FromString(Asset->GetName());
 	}
 	else
 	{
@@ -7018,16 +7020,25 @@ TSharedRef<SWidget> FPersonaMeshDetails::OnGenerateWidgetForClothingEntry(TShare
 
 FText FPersonaMeshDetails::OnGetClothingComboText(int32 InLodIdx, int32 InSectionIdx) const
 {
-	USkeletalMesh* Mesh = GetPersonaToolkit()->GetMesh();
+	const USkeletalMesh* const Mesh = GetPersonaToolkit()->GetMesh();
 
-	if(Mesh && !Mesh->IsCompiling())
+	if (Mesh && !Mesh->IsCompiling())
 	{
-		UClothingAssetCommon* ClothingAsset = Cast<UClothingAssetCommon>(Mesh->GetSectionClothingAsset(InLodIdx, InSectionIdx));
-
-		if(ClothingAsset && ClothingAsset->LodMap.IsValidIndex(InLodIdx))
+		if (const UClothingAssetBase* const ClothingAsset = Cast<const UClothingAssetBase>(Mesh->GetSectionClothingAsset(InLodIdx, InSectionIdx)))
 		{
-			const int32 ClothingLOD = ClothingAsset->LodMap[InLodIdx];
-			return FText::Format(LOCTEXT("ClothingAssetEntry_Name", "{0} - LOD{1}"), FText::FromString(ClothingAsset->GetName()), FText::AsNumber(ClothingLOD));
+			if (const UClothingAssetCommon* const ClothingAssetCommon = Cast<const UClothingAssetCommon>(Mesh->GetSectionClothingAsset(InLodIdx, InSectionIdx)))
+			{
+				if (ClothingAssetCommon->LodMap.IsValidIndex(InLodIdx))
+				{
+					const int32 ClothingLOD = ClothingAssetCommon->LodMap[InLodIdx];
+					return FText::Format(LOCTEXT("ClothingAssetEntry_Name", "{0} - LOD{1}"), FText::FromString(ClothingAsset->GetName()), FText::AsNumber(ClothingLOD));
+				}
+			}
+			else
+			{
+				// Non ClothingAssetCommon clothing data cannot have a different LOD index to the section's LOD index, so just show the name
+				return FText::FromString(ClothingAsset->GetName());
+			}
 		}
 	}
 
@@ -7062,7 +7073,7 @@ void FPersonaMeshDetails::OnClothingSelectionChanged(TSharedPtr<FClothingEntry>
 			OriginalSectionData.ClothingData.AssetGuid = FGuid();
 			OriginalSectionData.ClothingData.AssetLodIndex = INDEX_NONE;
 		};
-		if (UClothingAssetCommon* ClothingAsset = Cast<UClothingAssetCommon>(InNewEntry->Asset.Get()))
+		if (UClothingAssetBase* const ClothingAsset = InNewEntry->Asset.Get())
 		{
 			ClothingAsset->Modify();
 
@@ -7087,10 +7098,11 @@ void FPersonaMeshDetails::OnClothingSelectionChanged(TSharedPtr<FClothingEntry>
 			{
 				//Successful bind so set the SectionUserData
 				int32 AssetIndex = INDEX_NONE;
-				check(Mesh->GetMeshClothingAssets().Find(ClothingAsset, AssetIndex));
+				verify(Mesh->GetMeshClothingAssets().Find(ClothingAsset, AssetIndex));
+				const int32 AssetLodIndex = InNewEntry->AssetLodIndex != INDEX_NONE ? InNewEntry->AssetLodIndex : InLodIdx;  //  Use the mesh LOD index if the asset type doesn't allow independent LOD setting
 				OriginalSectionData.CorrespondClothAssetIndex = static_cast<int16>(AssetIndex);
 				OriginalSectionData.ClothingData.AssetGuid = ClothingAsset->GetAssetGuid();
-				OriginalSectionData.ClothingData.AssetLodIndex = InNewEntry->AssetLodIndex;
+				OriginalSectionData.ClothingData.AssetLodIndex = AssetLodIndex;
 			}
 		}
 		else if (Mesh)
diff --git a/Engine/Source/Editor/Persona/Private/SCreateClothingSettingsPanel.cpp b/Engine/Source/Editor/Persona/Private/SCreateClothingSettingsPanel.cpp
--- a/Engine/Source/Editor/Persona/Private/SCreateClothingSettingsPanel.cpp
+++ b/Engine/Source/Editor/Persona/Private/SCreateClothingSettingsPanel.cpp
@@ -13,7 +13,7 @@
 #include "DetailCategoryBuilder.h"
 #include "DetailWidgetRow.h"
 #include "PropertyCustomizationHelpers.h"
-#include "ClothingAssetBase.h"
+#include "ClothingAsset.h"
 #include "Widgets/Input/SComboButton.h"
 #include "Widgets/Text/STextBlock.h"
 #include "Framework/MultiBox/MultiBoxBuilder.h"
@@ -147,12 +147,12 @@ TSharedRef<SWidget> FClothCreateSettingsCustomization::OnGetTargetAssetMenu()
 
 	Builder.BeginSection(TEXT("TargetAssetDropdown"), LOCTEXT("TargetAssetMenuHeader", "Available Assets"));
 	{
-		if(USkeletalMesh* Mesh = MeshPtr.Get())
+		if (const USkeletalMesh* const Mesh = MeshPtr.Get())
 		{
 			const int32 NumCloths = Mesh->GetMeshClothingAssets().Num();
-			for(int32 ClothIndex = 0 ; ClothIndex < NumCloths ; ++ClothIndex)
+			for (int32 ClothIndex = 0; ClothIndex < NumCloths; ++ClothIndex)
 			{
-				if (UClothingAssetBase* ClothingAsset = Mesh->GetMeshClothingAssets()[ClothIndex])
+				if (const UClothingAssetCommon* const ClothingAsset = Cast<const UClothingAssetCommon>(Mesh->GetMeshClothingAssets()[ClothIndex]))
 				{
 					FUIAction Action;
 					Action.ExecuteAction = FExecuteAction::CreateSP(this, &FClothCreateSettingsCustomization::OnAssetSelected, ClothIndex);
@@ -202,13 +202,13 @@ TSharedRef<SWidget> FClothCreateSettingsCustomization::OnGetTargetLodMenu()
 {
 	check(ParamsStruct);
 
-	USkeletalMesh* Mesh = MeshPtr.Get();
+	const USkeletalMesh* const Mesh = MeshPtr.Get();
 
 	FMenuBuilder Builder(true, nullptr, TSharedPtr<FExtender>(), true);
 
 	Builder.BeginSection(TEXT("TargetLodDropdown"), LOCTEXT("TargetLodMenuHeader", "Available LODs"));
 	{
-		if(UClothingAssetBase* ClothingAsset = ParamsStruct->TargetAsset.Get())
+		if(const UClothingAssetCommon* const ClothingAsset = Cast<const UClothingAssetCommon>(ParamsStruct->TargetAsset.Get()))
 		{
 			const int32 NumLodEntries = ClothingAsset->GetNumLods();
 			for(int32 LodEntryIndex = 0; LodEntryIndex < NumLodEntries + 1; ++LodEntryIndex)
diff --git a/Engine/Source/Editor/SkeletalMeshEditor/Private/SkeletalMeshEditor.cpp b/Engine/Source/Editor/SkeletalMeshEditor/Private/SkeletalMeshEditor.cpp
--- a/Engine/Source/Editor/SkeletalMeshEditor/Private/SkeletalMeshEditor.cpp
+++ b/Engine/Source/Editor/SkeletalMeshEditor/Private/SkeletalMeshEditor.cpp
@@ -703,38 +703,47 @@ void FSkeletalMeshEditor::OnRemoveSectionFromLodAndBelowMenuItemClicked(int32 Lo
 
 void FSkeletalMeshEditor::FillApplyClothingAssetMenu(FMenuBuilder& MenuBuilder, int32 InLodIndex, int32 InSectionIndex)
 {
-	USkeletalMesh* Mesh = GetPersonaToolkit()->GetPreviewMesh();
-
-	// Nothing to fill
-	if(!Mesh)
+	if (const USkeletalMesh* const Mesh = GetPersonaToolkit()->GetPreviewMesh())
 	{
-		return;
-	}
-
-	MenuBuilder.BeginSection(TEXT("ApplyClothingMenu"), LOCTEXT("ApplyClothingMenuHeader", "Available Assets"));
-	{
-		for(UClothingAssetBase* BaseAsset : Mesh->GetMeshClothingAssets())
+		MenuBuilder.BeginSection(TEXT("ApplyClothingMenu"), LOCTEXT("ApplyClothingMenuHeader", "Available Assets"));
 		{
-			UClothingAssetCommon* ClothAsset = CastChecked<UClothingAssetCommon>(BaseAsset);
-
-			FUIAction Action;
-			Action.CanExecuteAction = FCanExecuteAction::CreateSP(this, &FSkeletalMeshEditor::CanApplyClothing, InLodIndex, InSectionIndex);
-
-			const int32 NumClothLods = ClothAsset->GetNumLods();
-			for(int32 ClothLodIndex = 0; ClothLodIndex < NumClothLods; ++ClothLodIndex)
+			for (UClothingAssetBase* const Asset : Mesh->GetMeshClothingAssets())
 			{
-				Action.ExecuteAction = FExecuteAction::CreateSP(this, &FSkeletalMeshEditor::OnApplyClothingAssetClicked, BaseAsset, InLodIndex, InSectionIndex, ClothLodIndex);
+				if (Asset)  // Asset could be null if the relevant plugin isn't loaded
+				{
+					FUIAction Action;
+					Action.CanExecuteAction = FCanExecuteAction::CreateSP(this, &FSkeletalMeshEditor::CanApplyClothing, InLodIndex, InSectionIndex);
 
-				MenuBuilder.AddMenuEntry(
-					FText::Format(LOCTEXT("ApplyClothingMenuItem", "{0} - LOD{1}"), FText::FromString(ClothAsset->GetName()), FText::AsNumber(ClothLodIndex)),
-					LOCTEXT("ApplyClothingMenuItem_ToolTip", "Apply this clothing asset to the selected mesh LOD and section"),
-					FSlateIcon(),
-					Action
-					);
+					if (const UClothingAssetCommon* const AssetCommon = Cast<UClothingAssetCommon>(Asset))
+					{
+						const int32 NumClothLods = AssetCommon->GetNumLods();
+						for (int32 ClothLodIndex = 0; ClothLodIndex < NumClothLods; ++ClothLodIndex)
+						{
+							Action.ExecuteAction = FExecuteAction::CreateSP(this, &FSkeletalMeshEditor::OnApplyClothingAssetClicked, Asset, InLodIndex, InSectionIndex, ClothLodIndex);
+
+							MenuBuilder.AddMenuEntry(
+								FText::Format(LOCTEXT("ApplyClothingMenuItem", "{0} - LOD{1}"), FText::FromString(Asset->GetName()), FText::AsNumber(ClothLodIndex)),
+								LOCTEXT("ApplyClothingMenuItem_ToolTip", "Apply this clothing data to the selected mesh LOD and section"),
+								FSlateIcon(),
+								Action);
+						}
+					}
+					else
+					{
+						const int32 ClothLodIndex = InLodIndex;  // Use the same LOD as the cloth for all other types of clothing data
+						Action.ExecuteAction = FExecuteAction::CreateSP(this, &FSkeletalMeshEditor::OnApplyClothingAssetClicked, Asset, InLodIndex, InSectionIndex, ClothLodIndex);
+
+						MenuBuilder.AddMenuEntry(
+							FText::FromString(Asset->GetName()),
+							LOCTEXT("ApplyClothingMenuItem_ToolTip", "Apply this clothing data to the selected mesh LOD and section"),
+							FSlateIcon(),
+							Action);
+					}
+				}
 			}
 		}
+		MenuBuilder.EndSection();
 	}
-	MenuBuilder.EndSection();
 }
 
 void FSkeletalMeshEditor::FillCreateClothingMenu(FMenuBuilder& MenuBuilder, int32 InLodIndex, int32 InSectionIndex)
@@ -966,9 +975,9 @@ void FSkeletalMeshEditor::ApplyClothing(UClothingAssetBase* InAsset, int32 InLod
 			OriginalSectionData.ClothingData.AssetGuid = FGuid();
 			OriginalSectionData.ClothingData.AssetLodIndex = INDEX_NONE;
 		};
-		if (UClothingAssetCommon* ClothingAsset = Cast<UClothingAssetCommon>(InAsset))
+		if (InAsset)
 		{
-			ClothingAsset->Modify();
+			InAsset->Modify();
 
 			// Look for a currently bound asset an unbind it if necessary first
 			if (UClothingAssetBase* CurrentAsset = Mesh->GetSectionClothingAsset(InLodIndex, InSectionIndex))
@@ -978,13 +987,13 @@ void FSkeletalMeshEditor::ApplyClothing(UClothingAssetBase* InAsset, int32 InLod
 				ClearOriginalSectionUserData();
 			}
 
-			if (ClothingAsset->BindToSkeletalMesh(Mesh, InLodIndex, InSectionIndex, InClothingLod))
+			if (InAsset->BindToSkeletalMesh(Mesh, InLodIndex, InSectionIndex, InClothingLod))
 			{
 				//Successful bind so set the SectionUserData
 				int32 AssetIndex = INDEX_NONE;
-				check(Mesh->GetMeshClothingAssets().Find(ClothingAsset, AssetIndex));
+				verify(Mesh->GetMeshClothingAssets().Find(InAsset, AssetIndex));
 				OriginalSectionData.CorrespondClothAssetIndex = static_cast<int16>(AssetIndex);
-				OriginalSectionData.ClothingData.AssetGuid = ClothingAsset->GetAssetGuid();
+				OriginalSectionData.ClothingData.AssetGuid = InAsset->GetAssetGuid();
 				OriginalSectionData.ClothingData.AssetLodIndex = InClothingLod;
 			}
 		}
diff --git a/Engine/Source/Runtime/ClothingSystemRuntimeInterface/Public/ClothingAssetBase.h b/Engine/Source/Runtime/ClothingSystemRuntimeInterface/Public/ClothingAssetBase.h
--- a/Engine/Source/Runtime/ClothingSystemRuntimeInterface/Public/ClothingAssetBase.h
+++ b/Engine/Source/Runtime/ClothingSystemRuntimeInterface/Public/ClothingAssetBase.h
@@ -52,7 +52,7 @@ public:
 	PURE_VIRTUAL(UClothingAssetBase::UpdateAllLODBiasMappings,);
 
 	/** Add a new LOD class instance. */
-	UE_DEPRECATED(5.7, "Use UClothingAssetCommone::AddNewLod() instead.")
+	UE_DEPRECATED(5.7, "Use UClothingAssetCommon::AddNewLod() instead.")
 	virtual int32 AddNewLod()
 	PURE_VIRTUAL(UClothingAssetBase::AddNewLod(), return INDEX_NONE;);
 
@@ -86,10 +86,12 @@ public:
 	PURE_VIRTUAL(UClothingAssetBase::RefreshBoneMapping, );
 
 	/** Check the validity of a LOD index */
+	UE_DEPRECATED(5.7, "Use UClothingAssetCommon::IsValidLod() instead.")
 	virtual bool IsValidLod(int32 InLodIndex) const
 	PURE_VIRTUAL(UClothingAssetBase::IsValidLod(), return false;);
 
 	/** Get the number of LODs defined in the clothing asset */
+	UE_DEPRECATED(5.7, "Use UClothingAssetCommon::GetNumLods() instead.")
 	virtual int32 GetNumLods() const
 	PURE_VIRTUAL(UClothingAssetBase::GetNumLods(), return 0;);
 
