Cloth - Improved the "Add Clothing" asset selector by enabling source clothing data asset type filters and enabling the plugin folders by default.

#rnx

diff --git a/Engine/Source/Editor/ClothPainter/Private/SClothAssetSelector.cpp b/Engine/Source/Editor/ClothPainter/Private/SClothAssetSelector.cpp
--- a/Engine/Source/Editor/ClothPainter/Private/SClothAssetSelector.cpp
+++ b/Engine/Source/Editor/ClothPainter/Private/SClothAssetSelector.cpp
@@ -1263,6 +1263,23 @@ TSharedRef<SWidget> SClothAssetSelector::OnGenerateSkeletalMeshPickerForClothCop
 	AssetPickerConfig.bFocusSearchBoxWhenOpened = true;
 	AssetPickerConfig.bShowBottomToolbar = false;
 	AssetPickerConfig.SelectionMode = ESelectionMode::Single;
+	AssetPickerConfig.bForceShowPluginContent = true;
+	AssetPickerConfig.SaveSettingsName = TEXT("ClothAssetSelectorAssetPicker");
+	AssetPickerConfig.bShowPathInColumnView = true;
+	AssetPickerConfig.bShowTypeInColumnView = true;
+	AssetPickerConfig.OnExtendAddFilterMenu = FOnExtendAddFilterMenu::CreateLambda([](UToolMenu* ToolMenu)
+		{
+			ToolMenu->RemoveSection("FilterBarFilterBasicAsset");  // Remove common filters section
+		});
+	AssetPickerConfig.bAddFilterUI = true;
+	AssetPickerConfig.Filter.bRecursiveClasses = true;
+	for (const UClothingAssetFactoryBase* const ClothingAssetFactory : ClothingEditorModule.GetClothingAssetFactories())
+	{
+		if (const UClass* const SupportedSourceAssetType = ClothingAssetFactory->GetSupportedSourceAssetType())
+		{
+			AssetPickerConfig.Filter.ClassPaths.Emplace(SupportedSourceAssetType->GetClassPathName());
+		}
+	}
 
 	return SNew(SBox)
 		.WidthOverride(300)
