Skeletal Mesh Cloth - Fixed Redo after undoing adding a clothing data from a Skeletal Mesh or Cloth Asset source.

[RN] Skeletal Mesh Cloth - Fixed Redo after undoing adding a clothing data from a Skeletal Mesh or Cloth Asset source.

diff --git a/Engine/Source/Editor/ClothPainter/Private/SClothAssetSelector.cpp b/Engine/Source/Editor/ClothPainter/Private/SClothAssetSelector.cpp
--- a/Engine/Source/Editor/ClothPainter/Private/SClothAssetSelector.cpp
+++ b/Engine/Source/Editor/ClothPainter/Private/SClothAssetSelector.cpp
@@ -1201,6 +1201,11 @@ void SClothAssetSelector::PostUndo(bool bSuccess)
 	OnRefresh();
 }
 
+void SClothAssetSelector::PostRedo(bool bSuccess)
+{
+	OnRefresh();
+}
+
 void SClothAssetSelector::OnCopyClothingAssetSelected(const FAssetData& AssetData)
 {
 	if (const UObject* const SourceAsset = AssetData.GetAsset())
diff --git a/Engine/Source/Editor/ClothPainter/Private/SClothAssetSelector.h b/Engine/Source/Editor/ClothPainter/Private/SClothAssetSelector.h
--- a/Engine/Source/Editor/ClothPainter/Private/SClothAssetSelector.h
+++ b/Engine/Source/Editor/ClothPainter/Private/SClothAssetSelector.h
@@ -81,6 +81,7 @@ public:
 
 	/** FEditorUndoClient interface */
 	virtual void PostUndo(bool bSuccess) override;
+	virtual void PostRedo(bool bSuccess) override;
 	/** End FEditorUndoClient interface */
 
 protected:
