Cloth - Deprecated GetDefaultClothingSimulationFactoryClass().

#rnx

diff --git a/Engine/Source/Editor/ClothPainter/Private/SClothAssetSelector.cpp b/Engine/Source/Editor/ClothPainter/Private/SClothAssetSelector.cpp
--- a/Engine/Source/Editor/ClothPainter/Private/SClothAssetSelector.cpp
+++ b/Engine/Source/Editor/ClothPainter/Private/SClothAssetSelector.cpp
@@ -431,14 +431,19 @@ public:
 		if(InColumnName == Column_CurrentTarget)
 		{
 			// Retrieve the mask names for the current clothing simulation factory
-			const TSubclassOf<class UClothingSimulationFactory> ClothingSimulationFactory = UClothingSimulationFactory::GetDefaultClothingSimulationFactoryClass();
-			if (ClothingSimulationFactory.Get() != nullptr)
+			if (const TSharedPtr<const SClothAssetSelector> AssetSelector = AssetSelectorPtr.Pin())
 			{
-				const UEnum* const Enum = ClothingSimulationFactory.GetDefaultObject()->GetWeightMapTargetEnum();
-				const FPointWeightMap* const Mask = Item->GetMask();
-				if(Enum && Mask)
+				if (const TStrongObjectPtr<const UClothingAssetBase> ClothingAsset = AssetSelector->GetSelectedAsset().Pin())
 				{
-					return SNew(STextBlock).Text(Enum->GetDisplayNameTextByValue((int64)Mask->CurrentTarget));
+					if (const UClothingSimulationFactory* const ClothingSimulationFactory = UClothingSimulationFactory::GetClothingSimulationFactory(ClothingAsset.Get()))
+					{
+						const UEnum* const Enum = ClothingSimulationFactory->GetWeightMapTargetEnum();
+						const FPointWeightMap* const Mask = Item->GetMask();
+						if (Enum && Mask)
+						{
+							return SNew(STextBlock).Text(Enum->GetDisplayNameTextByValue((int64)Mask->CurrentTarget));
+						}
+					}
 				}
 			}
 		}
@@ -630,21 +635,26 @@ private:
 		Builder.BeginSection(NAME_None, LOCTEXT("MaskTargets_SectionName", "Targets"));
 		{
 			// Retrieve the mask names for the current clothing simulation factory
-			const TSubclassOf<class UClothingSimulationFactory> ClothingSimulationFactory = UClothingSimulationFactory::GetDefaultClothingSimulationFactoryClass();
-			if (ClothingSimulationFactory.Get() != nullptr)
+			if (const TSharedPtr<const SClothAssetSelector> AssetSelector = AssetSelectorPtr.Pin())
 			{
-				if (const UEnum* const Enum = ClothingSimulationFactory.GetDefaultObject()->GetWeightMapTargetEnum())
+				if (const TStrongObjectPtr<const UClothingAssetBase> ClothingAsset = AssetSelector->GetSelectedAsset().Pin())
 				{
-					const int32 NumEntries = Enum->NumEnums();
-
-					// Iterate to -1 to skip the _MAX entry appended to the end of the enum
-					for(int32 Index = 0; Index < NumEntries - 1; ++Index)
+					if (const UClothingSimulationFactory* const ClothingSimulationFactory = UClothingSimulationFactory::GetClothingSimulationFactory(ClothingAsset.Get()))
 					{
-						FUIAction EntryAction(FExecuteAction::CreateSP(this, &SMaskListRow::OnSetTarget, (int32)Enum->GetValueByIndex(Index)));
+						if (const UEnum* const Enum = ClothingSimulationFactory->GetWeightMapTargetEnum())
+						{
+							const int32 NumEntries = Enum->NumEnums();
+
+							// Iterate to -1 to skip the _MAX entry appended to the end of the enum
+							for (int32 Index = 0; Index < NumEntries - 1; ++Index)
+							{
+								FUIAction EntryAction(FExecuteAction::CreateSP(this, &SMaskListRow::OnSetTarget, (int32)Enum->GetValueByIndex(Index)));
 
-						FText EntryText = Enum->GetDisplayNameTextByIndex(Index);
+								FText EntryText = Enum->GetDisplayNameTextByIndex(Index);
 
-						Builder.AddMenuEntry(EntryText, FText::GetEmpty(), FSlateIcon(), EntryAction);
+								Builder.AddMenuEntry(EntryText, FText::GetEmpty(), FSlateIcon(), EntryAction);
+							}
+						}
 					}
 				}
 			}
diff --git a/Engine/Source/Runtime/ClothingSystemRuntimeInterface/Private/ClothingSimulationFactory.cpp b/Engine/Source/Runtime/ClothingSystemRuntimeInterface/Private/ClothingSimulationFactory.cpp
--- a/Engine/Source/Runtime/ClothingSystemRuntimeInterface/Private/ClothingSimulationFactory.cpp
+++ b/Engine/Source/Runtime/ClothingSystemRuntimeInterface/Private/ClothingSimulationFactory.cpp
@@ -11,6 +11,7 @@ const FName IClothingSimulationFactoryClassProvider::FeatureName = TEXT("Clothin
 
 namespace ClothingSimulationFactoryConsoleVariables
 {
+	// UE_DEPRECATED(5.7, "Use GetClothingSimulationFactory() instead.")
 	TAutoConsoleVariable<FString> CVarDefaultClothingSimulationFactoryClass(
 		TEXT("p.Cloth.DefaultClothingSimulationFactoryClass"),
 		TEXT("ChaosClothingSimulationFactory"),  // Chaos is the default provider when Chaos Cloth is enabled
@@ -20,6 +21,7 @@ namespace ClothingSimulationFactoryConsoleVariables
 		, ECVF_Cheat);
 }
 
+// UE_DEPRECATED(5.7, "Use GetClothingSimulationFactory() instead.")
 UE_AUTORTFM_ALWAYS_OPEN
 TSubclassOf<class UClothingSimulationFactory> UClothingSimulationFactory::GetDefaultClothingSimulationFactoryClass()
 {
diff --git a/Engine/Source/Runtime/ClothingSystemRuntimeInterface/Public/ClothingSimulationFactory.h b/Engine/Source/Runtime/ClothingSystemRuntimeInterface/Public/ClothingSimulationFactory.h
--- a/Engine/Source/Runtime/ClothingSystemRuntimeInterface/Public/ClothingSimulationFactory.h
+++ b/Engine/Source/Runtime/ClothingSystemRuntimeInterface/Public/ClothingSimulationFactory.h
@@ -46,9 +46,7 @@ class UClothingSimulationFactory : public UObject
 	GENERATED_BODY()
 
 public:
-	// Return the default clothing simulation factory class as set by the build or by
-	// the p.Cloth.DefaultClothingSimulationFactoryClass console variable if any available.
-	// Otherwise return the last registered factory.
+	UE_DEPRECATED(5.7, "Use GetClothingSimulationFactory() instead.")
 	static UE_API TSubclassOf<class UClothingSimulationFactory> GetDefaultClothingSimulationFactoryClass();
 
 	// Return the clothing simulation factory for the specified asset.
