Skeletal Mesh Cloth:
- Fixed overzealous name uniqueness code that append _0 to all clothing data names.
- Added missing sanitization when renaming Clothing Data that could cause issues with further renaming.
- Fixed name reuse when a clothing data is deleted and the name needs to be used again.

[RN] Skeletal Mesh Cloth - Fixed an issue where all new clothing data names had _0 appended. ?_0? appended.

diff --git a/Engine/Plugins/ChaosClothAssetEditor/Source/ChaosClothAssetTools/ChaosClothAssetTools.Build.cs b/Engine/Plugins/ChaosClothAssetEditor/Source/ChaosClothAssetTools/ChaosClothAssetTools.Build.cs
--- a/Engine/Plugins/ChaosClothAssetEditor/Source/ChaosClothAssetTools/ChaosClothAssetTools.Build.cs
+++ b/Engine/Plugins/ChaosClothAssetEditor/Source/ChaosClothAssetTools/ChaosClothAssetTools.Build.cs
@@ -31,6 +31,7 @@ public class ChaosClothAssetTools : ModuleRules
 				"UnrealEd",
 				"Projects",
 				"ClothPainter",                 // For clothing asset exporter
+				"ClothingSystemRuntimeInterface",
 				"ClothingSystemRuntimeCommon",
 				"ClothingSystemEditorInterface",  // For UChaosClothAssetSKMClothingAssetFactory
 				"GeometryCore",					// For DynamicMesh conversion
diff --git a/Engine/Plugins/ChaosClothAssetEditor/Source/ChaosClothAssetTools/Private/ChaosClothAsset/ClothAssetSKMClothingAssetFactory.cpp b/Engine/Plugins/ChaosClothAssetEditor/Source/ChaosClothAssetTools/Private/ChaosClothAsset/ClothAssetSKMClothingAssetFactory.cpp
--- a/Engine/Plugins/ChaosClothAssetEditor/Source/ChaosClothAssetTools/Private/ChaosClothAsset/ClothAssetSKMClothingAssetFactory.cpp
+++ b/Engine/Plugins/ChaosClothAssetEditor/Source/ChaosClothAssetTools/Private/ChaosClothAsset/ClothAssetSKMClothingAssetFactory.cpp
@@ -22,7 +22,9 @@ TArray<UClothingAssetBase*> UChaosClothAssetSKMClothingAssetFactory::CreateFromS
 	TArray<UClothingAssetBase*> ClothingAssets;
 	if (const UChaosClothAssetBase* const ClothAsset = Cast<UChaosClothAssetBase>(SourceAsset))
 	{
-		const FName AssetName = MakeUniqueObjectName(TargetMesh, UChaosClothAssetSKMClothingAsset::StaticClass(), ClothAsset->GetFName());
+		const FName AssetName = StaticFindObject(UClothingAssetBase::StaticClass(), TargetMesh, *ClothAsset->GetName()) ?
+			MakeUniqueObjectName(TargetMesh, UClothingAssetBase::StaticClass(), ClothAsset->GetFName()) :
+			ClothAsset->GetFName();
 		UChaosClothAssetSKMClothingAsset* const ClothAssetSKMClothingAsset = NewObject<UChaosClothAssetSKMClothingAsset>(TargetMesh, AssetName);
 		ClothAssetSKMClothingAsset->SetAsset(Cast<UChaosClothAssetBase>(SourceAsset));
 		ClothingAssets.Emplace(ClothAssetSKMClothingAsset);
diff --git a/Engine/Source/Editor/ClothPainter/Private/SClothAssetSelector.cpp b/Engine/Source/Editor/ClothPainter/Private/SClothAssetSelector.cpp
--- a/Engine/Source/Editor/ClothPainter/Private/SClothAssetSelector.cpp
+++ b/Engine/Source/Editor/ClothPainter/Private/SClothAssetSelector.cpp
@@ -43,6 +43,7 @@
 #include "Misc/AssertionMacros.h"
 #include "Misc/Attribute.h"
 #include "Modules/ModuleManager.h"
+#include "ObjectTools.h"
 #include "PointWeightMap.h"
 #include "Rendering/SkeletalMeshLODModel.h"
 #include "Rendering/SkeletalMeshModel.h"
@@ -167,19 +168,19 @@ public:
 		{
 			if(UClothingAssetBase* const Asset = Item->ClothingAsset.Get())
 			{
-				FText TrimText = FText::TrimPrecedingAndTrailing(InText);
-
-				if(Asset->GetName() != TrimText.ToString())
+				if (const FText TrimText = FText::TrimPrecedingAndTrailing(InText); !TrimText.IsEmpty())
 				{
-					FName NewName(*TrimText.ToString());
+					const FString SanitizedName = ObjectTools::SanitizeObjectName(TrimText.ToString());
 
-					// Check for an existing object, and if we find one build a unique name based on the request
-					if(StaticFindObject(UClothingAssetBase::StaticClass(), Asset->GetOuter(), *NewName.ToString()))
+					if(Asset->GetName() != SanitizedName)
 					{
-						NewName = MakeUniqueObjectName(Asset->GetOuter(), UClothingAssetBase::StaticClass(), FName(*TrimText.ToString()));
-					}
+						// Check for an existing object, and if we find one build a unique name based on the request
+						const FString NewName = StaticFindObject(UClothingAssetBase::StaticClass(), Asset->GetOuter(), *SanitizedName) ?
+							MakeUniqueObjectName(Asset->GetOuter(), UClothingAssetBase::StaticClass(), FName(SanitizedName)).ToString() :
+							SanitizedName;
 
-					Asset->Rename(*NewName.ToString(), Asset->GetOuter());
+						Asset->Rename(*NewName, Asset->GetOuter());
+					}
 				}
 			}
 		}
@@ -262,7 +263,7 @@ private:
 			SourceSectionUserData.ClothingData = LODModel.Sections[SectionIndex].ClothingData;
 		};
 
-		if(UClothingAssetBase* const Asset = Item->ClothingAsset.Get())
+		if(UClothingAssetBase* Asset = Item->ClothingAsset.Get())
 		{
 			if(USkeletalMesh* SkelMesh = Cast<USkeletalMesh>(Asset->GetOuter()))
 			{
@@ -276,6 +277,8 @@ private:
 
 					Asset->UnbindFromSkeletalMesh(SkelMesh);
 					SkelMesh->GetMeshClothingAssets().RemoveAt(AssetIndex);
+					Asset->ConditionalBeginDestroy();  // Mark the clothing data as to be destroyed asap, so that the object name can be reused if needed
+					Asset = nullptr;
 
 					// Need to fix up asset indices on sections.
 					if(FSkeletalMeshModel* Model = SkelMesh->GetImportedModel())
diff --git a/Engine/Source/Editor/ClothingSystemEditor/Private/ClothingAssetFactory.cpp b/Engine/Source/Editor/ClothingSystemEditor/Private/ClothingAssetFactory.cpp
--- a/Engine/Source/Editor/ClothingSystemEditor/Private/ClothingAssetFactory.cpp
+++ b/Engine/Source/Editor/ClothingSystemEditor/Private/ClothingAssetFactory.cpp
@@ -167,8 +167,10 @@ UClothingAssetBase* UClothingAssetFactory::CreateFromSkeletalMesh(USkeletalMesh*
 		return nullptr;
 	}
 
-	FString SanitizedName = ObjectTools::SanitizeObjectName(Params.AssetName);
-	FName ObjectName = MakeUniqueObjectName(TargetMesh, UClothingAssetCommon::StaticClass(), FName(*SanitizedName));
+	const FString SanitizedName = ObjectTools::SanitizeObjectName(Params.AssetName);
+	const FName ObjectName = StaticFindObject(UClothingAssetBase::StaticClass(), TargetMesh, *SanitizedName) ?
+		MakeUniqueObjectName(TargetMesh, UClothingAssetBase::StaticClass(), FName(SanitizedName)) :
+		FName(SanitizedName);
 
 	UClothingAssetCommon* NewAsset = NewObject<UClothingAssetCommon>(TargetMesh, ObjectName);
 	NewAsset->SetFlags(RF_Transactional);
@@ -214,7 +216,10 @@ UClothingAssetBase* UClothingAssetFactory::CreateFromExistingCloth(USkeletalMesh
 	}
 
 	//Duplicating the clothing asset using the existing asset as a template
-	const FName NewAssetName = MakeUniqueObjectName(TargetMesh, UClothingAssetCommon::StaticClass(), SourceClothingAsset->GetFName());
+	const FName NewAssetName = StaticFindObject(UClothingAssetBase::StaticClass(), TargetMesh, *SourceClothingAsset->GetName()) ?
+		MakeUniqueObjectName(TargetMesh, UClothingAssetBase::StaticClass(), SourceClothingAsset->GetFName()) :
+		SourceClothingAsset->GetFName();
+
 	UClothingAssetCommon* const NewAsset = DuplicateObject<UClothingAssetCommon>(SourceClothingAsset, TargetMesh, NewAssetName);
 
 	NewAsset->AssetGuid = FGuid::NewGuid();
