Skeletal Mesh Cloth - Fixed a bug where the Clothing Data detail view would persist after the entry is deleted.
[RN] Skeletal Mesh Cloth - Fixed a bug where the Clothing Data detail view would persist after the entry is deleted.

diff --git a/Engine/Source/Editor/ClothPainter/Private/SClothAssetSelector.cpp b/Engine/Source/Editor/ClothPainter/Private/SClothAssetSelector.cpp
--- a/Engine/Source/Editor/ClothPainter/Private/SClothAssetSelector.cpp
+++ b/Engine/Source/Editor/ClothPainter/Private/SClothAssetSelector.cpp
@@ -1389,6 +1389,10 @@ void SClothAssetSelector::OnAssetListSelectionChanged(TSharedPtr<FClothingAssetL
 	{
 		SetSelectedAsset(InSelectedItem->ClothingAsset);
 	}
+	else
+	{
+		SetSelectedAsset(nullptr);
+	}
 }
 
 TSharedRef<ITableRow> SClothAssetSelector::OnGenerateWidgetForMaskItem(TSharedPtr<FClothingMaskListItem> InItem, const TSharedRef<STableViewBase>& OwnerTable)
