Skeletal Mesh Cloth - Preserve Direct mode behavior when changing clothing data selections.

#rnx

diff --git a/Engine/Source/Editor/ClothPainter/Private/SClothAssetSelector.cpp b/Engine/Source/Editor/ClothPainter/Private/SClothAssetSelector.cpp
--- a/Engine/Source/Editor/ClothPainter/Private/SClothAssetSelector.cpp
+++ b/Engine/Source/Editor/ClothPainter/Private/SClothAssetSelector.cpp
@@ -1385,13 +1385,16 @@ TSharedRef<ITableRow> SClothAssetSelector::OnGenerateWidgetForClothingAssetItem(
 
 void SClothAssetSelector::OnAssetListSelectionChanged(TSharedPtr<FClothingAssetListItem> InSelectedItem, ESelectInfo::Type InSelectInfo)
 {
-	if(InSelectedItem.IsValid() && InSelectInfo != ESelectInfo::Direct)
+	if (InSelectInfo != ESelectInfo::Direct)
 	{
-		SetSelectedAsset(InSelectedItem->ClothingAsset);
-	}
-	else
-	{
-		SetSelectedAsset(nullptr);
+		if (InSelectedItem.IsValid())
+		{
+			SetSelectedAsset(InSelectedItem->ClothingAsset);
+		}
+		else
+		{
+			SetSelectedAsset(nullptr);
+		}
 	}
 }
 
