SkeletalMeshEditor CLothing Data panel : add Rename action to the context menu

#Synced-CL 49114577

diff --git a/Engine/Source/Editor/ClothPainter/Private/SClothAssetSelector.cpp b/Engine/Source/Editor/ClothPainter/Private/SClothAssetSelector.cpp
--- a/Engine/Source/Editor/ClothPainter/Private/SClothAssetSelector.cpp
+++ b/Engine/Source/Editor/ClothPainter/Private/SClothAssetSelector.cpp
@@ -196,6 +196,11 @@ public:
 
 		const FClothingAssetListCommands& Commands = FClothingAssetListCommands::Get();
 
+		UICommandList->MapAction(
+			FGenericCommands::Get().Rename,
+			FExecuteAction::CreateSP(this, &SAssetListRow::RenameAsset)
+		);
+
 		UICommandList->MapAction(
 			FGenericCommands::Get().Delete,
 			FExecuteAction::CreateSP(this, &SAssetListRow::DeleteAsset)
@@ -237,6 +242,8 @@ public:
 
 			Builder.BeginSection(NAME_None, LOCTEXT("AssetActions_SectionName", "Actions"));
 			{
+				Builder.AddMenuEntry(FGenericCommands::Get().Rename);
+				Builder.AddSeparator();
 				Builder.AddMenuEntry(FGenericCommands::Get().Delete);
 				if (LastLodIndex > 0)  // Can't delete LOD0
 				{
@@ -263,6 +270,15 @@ public:
 	}
 
 private:
+	
+	void RenameAsset()
+	{
+		if (EditableText)
+		{
+			EditableText->EnterEditingMode();
+		}
+	}
+
 
 	void DeleteAsset()
 	{
